Slate: Refactored core Slate implementation into SlateCore module in preparation for UMG.

Other Updates:
- The WidgetReflector is now in its own module as well. It will be converted to a plug-in later.
- The Public API of both Slate and SlateCore has largely been reorganized for better discoverabilty. More cleanup work is needed.
- Added a lot of missing API documentation and fixed existing ones. More and better documentation is needed.
- Removed dead code, fixed a couple things I stubled upon, and conformed to coding guidelines (NULL vs nullptr, line breaks, etc.)

Upgrade Notes:
- The Slate Remote Server is currently disabled - will be re-enabled shortly!
- If your module previously had a module dependency to 'Slate', it now also needs a PrivateModuleDependency to 'SlateCore' in its Build.cs file.
- If your module exposes in any of its Public header files types that are now declared in SlateCore, it needs a PublicModuleDependency to 'SlateCore'
- The ToolTip property type on SWidget has changed from SToolTip to IToolTip; change local variables to TSharedPtr<IToolTip> instead of TSharedPtr<SToolTip> where needed
- IToolTip is not a widget. If you need access to the actual widget that represents the tool tip, use IToolTip::AsWidget(); If you need access to the tool tip's content, use IToolTip::GetContentWidget()

Troubleshooting:
- After syncing to this changelist you may have to clean your /Engine/Intermediate/Build/ directory and rebuild your entire project
- If in your project you are getting linker errors for unresolved types that are now declared in SlateCore, you may be missing a dependency to 'SlateCore'
- If in the Engine code you are getting linker errors for unresolved types that are now declared in SlateCore, you may need to rebuild the entire Engine

diff --git a/Engine/Plugins/Developer/PerforceSourceControl/Source/PerforceSourceControl/PerforceSourceControl.Build.cs b/Engine/Plugins/Developer/PerforceSourceControl/Source/PerforceSourceControl/PerforceSourceControl.Build.cs
--- a/Engine/Plugins/Developer/PerforceSourceControl/Source/PerforceSourceControl/PerforceSourceControl.Build.cs
+++ b/Engine/Plugins/Developer/PerforceSourceControl/Source/PerforceSourceControl/PerforceSourceControl.Build.cs
@@ -12,10 +12,11 @@ public class PerforceSourceControl : ModuleRules
 				"CoreUObject",
                 "InputCore",
 				"Slate",
+				"SlateCore",
                 "EditorStyle",
 				"SourceControl",
 			}
-            );
+		);
 
 		AddThirdPartyPrivateStaticDependencies(Target, "Perforce");
 
diff --git a/Engine/Plugins/Editor/EpicSurvey/Source/EpicSurvey/EpicSurvey.Build.cs b/Engine/Plugins/Editor/EpicSurvey/Source/EpicSurvey/EpicSurvey.Build.cs
--- a/Engine/Plugins/Editor/EpicSurvey/Source/EpicSurvey/EpicSurvey.Build.cs
+++ b/Engine/Plugins/Editor/EpicSurvey/Source/EpicSurvey/EpicSurvey.Build.cs
@@ -27,6 +27,7 @@ namespace UnrealBuildTool.Rules
 				{
 					"HTTP",
 					"OnlineSubSystem",
+					"SlateCore",
 					"UnrealEd",
 				}
 			);
diff --git a/Engine/Plugins/Editor/IntroTutorials/Source/IntroTutorials/IntroTutorials.Build.cs b/Engine/Plugins/Editor/IntroTutorials/Source/IntroTutorials/IntroTutorials.Build.cs
--- a/Engine/Plugins/Editor/IntroTutorials/Source/IntroTutorials/IntroTutorials.Build.cs
+++ b/Engine/Plugins/Editor/IntroTutorials/Source/IntroTutorials/IntroTutorials.Build.cs
@@ -31,7 +31,8 @@ namespace UnrealBuildTool.Rules
 				new string[] {				
                     "UnrealEd",
                     "Kismet",
-                    "PlacementMode"
+                    "PlacementMode",
+					"SlateCore",
 				}
 			);
 
diff --git a/Engine/Plugins/Editor/PluginsEditor/Source/PluginsEditor/PluginsEditor.Build.cs b/Engine/Plugins/Editor/PluginsEditor/Source/PluginsEditor/PluginsEditor.Build.cs
--- a/Engine/Plugins/Editor/PluginsEditor/Source/PluginsEditor/PluginsEditor.Build.cs
+++ b/Engine/Plugins/Editor/PluginsEditor/Source/PluginsEditor/PluginsEditor.Build.cs
@@ -12,7 +12,8 @@ namespace UnrealBuildTool.Rules
 					"CoreUObject",		// @todo Mac: for some reason CoreUObject and Engine are needed to link in debug on Mac
                     "InputCore",
 					"Engine",
-					"Slate"
+					"Slate",
+					"SlateCore",
 				}
 			);
 			
diff --git a/Engine/Plugins/Experimental/CodeView/Source/CodeView/CodeView.Build.cs b/Engine/Plugins/Experimental/CodeView/Source/CodeView/CodeView.Build.cs
--- a/Engine/Plugins/Experimental/CodeView/Source/CodeView/CodeView.Build.cs
+++ b/Engine/Plugins/Experimental/CodeView/Source/CodeView/CodeView.Build.cs
@@ -16,16 +16,17 @@ namespace UnrealBuildTool.Rules
 					"Engine",
 					"UnrealEd",
 					"EditorStyle",
-					"PropertyEditor"
+					"PropertyEditor",
 				}
-				);
+			);
 
 			PrivateDependencyModuleNames.AddRange(
 				new string[]
 				{
-					"DetailCustomizations"
+					"DetailCustomizations",
+					"SlateCore",
 				}
-				);
+			);
 		}
 	}
 }
\ No newline at end of file
diff --git a/Engine/Plugins/Experimental/LiveEditor/Source/LiveEditor/LiveEditor.Build.cs b/Engine/Plugins/Experimental/LiveEditor/Source/LiveEditor/LiveEditor.Build.cs
--- a/Engine/Plugins/Experimental/LiveEditor/Source/LiveEditor/LiveEditor.Build.cs
+++ b/Engine/Plugins/Experimental/LiveEditor/Source/LiveEditor/LiveEditor.Build.cs
@@ -10,18 +10,17 @@ namespace UnrealBuildTool.Rules
 				new string[] {
 					// ... add public include paths required here ...
 				}
-				);
+			);
 
 			PrivateIncludePaths.AddRange(
 				new string[] {
 					"LiveEditor/Private",
 					// ... add other private include paths required here ...
 				}
-				);
+			);
 
 			PublicDependencyModuleNames.AddRange(
-				new string[]
-				{
+				new string[] {
                     "Core",
 					"CoreUObject",
                     "Slate",
@@ -39,21 +38,19 @@ namespace UnrealBuildTool.Rules
                     "WorkspaceMenuStructure",//@TODO: This makes it depend on a module in editor, even if that module does not depend on other editor stuff!
                     "LiveEditorListenServer"
 				}
-				);
+			);
 
 			PrivateDependencyModuleNames.AddRange(
-				new string[]
-				{
-					// ... add private dependencies that you statically link with here ...
+				new string[] {
+					"SlateCore",
 				}
-				);
+			);
 
 			DynamicallyLoadedModuleNames.AddRange(
-				new string[]
-				{
+				new string[] {
 					// ... add any modules that your module loads dynamically here ...
 				}
-				);
+			);
 		}
 	}
 }
\ No newline at end of file
diff --git a/Engine/Plugins/Experimental/Paper2D/Source/Paper2DEditor/Paper2DEditor.Build.cs b/Engine/Plugins/Experimental/Paper2D/Source/Paper2DEditor/Paper2DEditor.Build.cs
--- a/Engine/Plugins/Experimental/Paper2D/Source/Paper2DEditor/Paper2DEditor.Build.cs
+++ b/Engine/Plugins/Experimental/Paper2D/Source/Paper2DEditor/Paper2DEditor.Build.cs
@@ -17,6 +17,7 @@ public class Paper2DEditor : ModuleRules
 				"Core",
 				"CoreUObject",
 				"Slate",
+				"SlateCore",
 				"Engine",
                 "InputCore",
 				"AssetTools",
diff --git a/Engine/Plugins/Runtime/AppleMoviePlayer/Source/AppleMoviePlayer/AppleMoviePlayer.Build.cs b/Engine/Plugins/Runtime/AppleMoviePlayer/Source/AppleMoviePlayer/AppleMoviePlayer.Build.cs
--- a/Engine/Plugins/Runtime/AppleMoviePlayer/Source/AppleMoviePlayer/AppleMoviePlayer.Build.cs
+++ b/Engine/Plugins/Runtime/AppleMoviePlayer/Source/AppleMoviePlayer/AppleMoviePlayer.Build.cs
@@ -6,11 +6,8 @@ namespace UnrealBuildTool.Rules
 	{
         public AppleMoviePlayer(TargetInfo Target)
 		{
-			PrivateIncludePaths.Add("Runtime/AppleMoviePlayer/Private");
-
 			PublicDependencyModuleNames.AddRange(
-				new string[]
-				{
+				new string[] {
 					"Core",
 				    "CoreUObject",
 				    "Engine",
@@ -19,11 +16,23 @@ namespace UnrealBuildTool.Rules
                     "RHI",
                     "Slate"
 				}
-				);
-				
+			);
+
+			PrivateDependencyModuleNames.AddRange(
+				new string[] {
+					"SlateCore",
+				}
+			);
+
+			PrivateIncludePaths.Add("Runtime/AppleMoviePlayer/Private");
+
 			if (Target.Platform == UnrealTargetPlatform.Mac)
 			{
-				PublicFrameworks.AddRange(new string[] { "QuartzCore" });
+				PublicFrameworks.AddRange(
+					new string[] {
+						"QuartzCore"
+					}
+				);
 			}
 		}
 	}
diff --git a/Engine/Source/Developer/CrashTracker/CrashTracker.Build.cs b/Engine/Source/Developer/CrashTracker/CrashTracker.Build.cs
--- a/Engine/Source/Developer/CrashTracker/CrashTracker.Build.cs
+++ b/Engine/Source/Developer/CrashTracker/CrashTracker.Build.cs
@@ -17,7 +17,8 @@ public class CrashTracker : ModuleRules
 				"RenderCore",
 				"RHI",
 				"ShaderCore",
-				"Slate"
+				"Slate",
+				"SlateCore",
 			}
 		);
 	}
diff --git a/Engine/Source/Developer/CrashTracker/Private/CrashTrackerPrivatePCH.h b/Engine/Source/Developer/CrashTracker/Private/CrashTrackerPrivatePCH.h
--- a/Engine/Source/Developer/CrashTracker/Private/CrashTrackerPrivatePCH.h
+++ b/Engine/Source/Developer/CrashTracker/Private/CrashTrackerPrivatePCH.h
@@ -16,7 +16,7 @@
 #include "ImageWrapper.h"
 #include "RenderCore.h"
 #include "RenderResource.h"
-#include "Renderer.h"
+#include "Slate.h"
 
 
 /* Private includes
diff --git a/Engine/Source/Developer/MeshUtilities/MeshUtilities.Build.cs b/Engine/Source/Developer/MeshUtilities/MeshUtilities.Build.cs
--- a/Engine/Source/Developer/MeshUtilities/MeshUtilities.Build.cs
+++ b/Engine/Source/Developer/MeshUtilities/MeshUtilities.Build.cs
@@ -7,17 +7,22 @@ public class MeshUtilities : ModuleRules
 {
 	public MeshUtilities(TargetInfo Target)
 	{
-		PrivateDependencyModuleNames.Add("Core");
-		PrivateDependencyModuleNames.Add("CoreUObject");
-		PrivateDependencyModuleNames.Add("Engine");
-		PrivateDependencyModuleNames.Add("RenderCore"); // For FPackedNormal.
-        PrivateDependencyModuleNames.Add("Slate");      // For FSlateTextureAtlas
-        PrivateDependencyModuleNames.Add("UnrealEd");
+		PrivateDependencyModuleNames.AddRange(
+			new string [] {
+				"Core",
+				"CoreUObject",
+				"Engine",
+				"RawMesh",
+				"RenderCore", // For FPackedNormal
+				"Slate",      // For FSlateTextureAtlas
+				"SlateCore",
+				"UnrealEd",
+			}
+		);
+
 		AddThirdPartyPrivateStaticDependencies(Target, "nvTriStrip");
 		AddThirdPartyPrivateStaticDependencies(Target, "ForsythTriOptimizer");
-		AddThirdPartyPrivateStaticDependencies(Target, "MeshSimplifier");
-		PrivateDependencyModuleNames.Add("RawMesh");
-        
+		AddThirdPartyPrivateStaticDependencies(Target, "MeshSimplifier");      
 
 		if ((Target.Platform == UnrealTargetPlatform.Win64) || (Target.Platform == UnrealTargetPlatform.Win32))
 		{
@@ -29,4 +34,4 @@ public class MeshUtilities : ModuleRules
 			AddThirdPartyPrivateDynamicDependencies(Target, "SimplygonMeshReduction");
 		}
 	}
-}
\ No newline at end of file
+}
diff --git a/Engine/Source/Developer/MessageLog/MessageLog.Build.cs b/Engine/Source/Developer/MessageLog/MessageLog.Build.cs
--- a/Engine/Source/Developer/MessageLog/MessageLog.Build.cs
+++ b/Engine/Source/Developer/MessageLog/MessageLog.Build.cs
@@ -13,25 +13,26 @@ public class MessageLog : ModuleRules
 				"Developer/MessageLog/Private/UserInterface",
                 "Developer/MessageLog/Private/Model",
 			}
-            );
+        );
 
 		PrivateDependencyModuleNames.AddRange(
 			new string[] {
 				"Core",
 				"CoreUObject",
 				"Slate",
+				"SlateCore",
                 "InputCore",
-                "EditorStyle"
+                "EditorStyle",
 			}
-			);
+		);
 
 		if (UEBuildConfiguration.bBuildEditor)
 		{
 			PrivateDependencyModuleNames.AddRange(
 				new string[] {
-					"Documentation"
+					"Documentation",
 				}
-				);		
+			);		
 		}
 	}
 }
diff --git a/Engine/Source/Developer/MessageLog/Private/MessageLogPrivatePCH.h b/Engine/Source/Developer/MessageLog/Private/MessageLogPrivatePCH.h
--- a/Engine/Source/Developer/MessageLog/Private/MessageLogPrivatePCH.h
+++ b/Engine/Source/Developer/MessageLog/Private/MessageLogPrivatePCH.h
@@ -6,7 +6,6 @@
 #include "Core.h"
 #include "Slate.h"
 #include "Attribute.h"
-#include "PrimitiveTypes.h"
 #include "EditorStyle.h"
 
 #endif // __MessageLogPrivatePCH_h__
diff --git a/Engine/Source/Developer/ScreenShotComparison/ScreenShotComparison.Build.cs b/Engine/Source/Developer/ScreenShotComparison/ScreenShotComparison.Build.cs
--- a/Engine/Source/Developer/ScreenShotComparison/ScreenShotComparison.Build.cs
+++ b/Engine/Source/Developer/ScreenShotComparison/ScreenShotComparison.Build.cs
@@ -6,13 +6,19 @@ public class ScreenShotComparison : ModuleRules
 {
 	public ScreenShotComparison(TargetInfo Target)
 	{
-		PrivateDependencyModuleNames.AddRange(
+		PublicDependencyModuleNames.AddRange(
 			new string[] {
 				"Core",
-                "InputCore",
-				"Slate",
+			}
+		);
+
+		PrivateDependencyModuleNames.AddRange(
+			new string[] {
                 "EditorStyle",
+                "InputCore",
 				"ScreenShotComparisonTools",
+				"Slate",
+				"SlateCore",
 			}
 		);
 
diff --git a/Engine/Source/Developer/SlateReflector/Private/Models/WidgetReflectorNode.h b/Engine/Source/Developer/SlateReflector/Private/Models/WidgetReflectorNode.h
new file mode 100644
--- /dev/null
+++ b/Engine/Source/Developer/SlateReflector/Private/Models/WidgetReflectorNode.h
@@ -0,0 +1,103 @@
+// Copyright 1998-2014 Epic Games, Inc. All Rights Reserved.
+
+/*=============================================================================
+	WidgetReflectorNode.h: Declares the FWidgetReflectorNode class.
+=============================================================================*/
+
+#pragma once
+
+
+/**
+  * The WidgetReflector captures layout information and widget hierarchy structure into a tree of FReflectorNode.
+  * The node also contains useful visualization and debug info.
+  */
+struct FReflectorNode
+{
+	/** The widget which this node describes */
+	TWeakPtr<SWidget> Widget;
+
+	/** The geometry of the widget */
+	FGeometry Geometry;
+
+	/** ReflectorNodes for the Widget's children. */
+	TArray<TSharedPtr<FReflectorNode>> ChildNodes;	
+
+	/** A tint that is applied to text in order to provide visual hints */
+	FLinearColor Tint;
+
+	/** Should we visualize this node */
+	bool bVisualizeThisNode;
+
+public:
+
+	/**
+	 * FReflectorNodes must be constructed as shared pointers.
+	 *
+	 * @param InWidgetGeometry   Optional widget and associated geometry which this node should represent
+	 */
+	static TSharedRef<FReflectorNode> New( const FArrangedWidget& InWidgetGeometry = FArrangedWidget(SNullWidget::NullWidget, FGeometry()) )
+	{
+		return MakeShareable( new FReflectorNode(InWidgetGeometry) );
+	}
+
+	/**
+	 * Capture all the children of the supplied widget along with their layout results
+	 * Note that we include both visible and invisible children!
+	 * 
+	 * @param InWidgetGeometry   Widget and geometry whose children to capture in the snapshot.
+	 */
+	static TSharedRef<FReflectorNode> NewTreeFrom( const FArrangedWidget& InWidgetGeometry )
+	{
+		TSharedRef<FReflectorNode> NewNode = MakeShareable( new FReflectorNode(InWidgetGeometry) );
+
+		FArrangedChildren ArrangedChildren(EVisibility::All);
+		NewNode->Widget.Pin()->ArrangeChildren( NewNode->Geometry, ArrangedChildren );
+		
+		for (int32 WidgetIndex=0; WidgetIndex < ArrangedChildren.Num(); ++WidgetIndex)
+		{
+			// Note that we include both visible and invisible children!
+			NewNode->ChildNodes.Add( FReflectorNode::NewTreeFrom( ArrangedChildren(WidgetIndex) )  );
+		}
+
+		return NewNode;
+	}
+
+	/**
+	 * Locate all the widgets from a widget path in a list of ReflectorNodes and their children.
+	 *
+	 * @param CandidateNodes    A list of FReflectorNodes that represent widgets.
+	 * @param WidgetPathToFind  We want to find all reflector nodes corresponding to widgets in this path
+	 * @param SearchResult      An array that gets results put in it
+	 * @param NodeIndexToFind   Index of the widget in the path that we are currently looking for; we are done when we've found all of them
+	 */
+	static void FindWidgetPath( const TArray<TSharedPtr<FReflectorNode>>& CandidateNodes, const FWidgetPath& WidgetPathToFind, TArray< TSharedPtr<FReflectorNode> >& SearchResult, int32 NodeIndexToFind = 0 )
+	{
+		if (NodeIndexToFind < WidgetPathToFind.Widgets.Num())
+		{
+			const FArrangedWidget& WidgetToFind = WidgetPathToFind.Widgets(NodeIndexToFind);
+
+			for (int32 NodeIndex=0; NodeIndex < CandidateNodes.Num(); ++NodeIndex)
+			{
+				if (CandidateNodes[NodeIndex]->Widget.Pin() == WidgetPathToFind.Widgets(NodeIndexToFind).Widget)
+				{
+					SearchResult.Add(CandidateNodes[NodeIndex]);
+					FindWidgetPath(CandidateNodes[NodeIndex]->ChildNodes, WidgetPathToFind, SearchResult, NodeIndexToFind + 1);
+				}
+			}
+		}
+	}
+
+protected:
+
+	/**
+	 * FReflectorNode must be constructor through static methods.
+	 *
+	 * @param InWidgetGeometry   Widget and associated geometry that this node will represent
+	 */
+	FReflectorNode( const FArrangedWidget& InWidgetGeometry )
+		: Widget( InWidgetGeometry.Widget )
+		, Geometry( InWidgetGeometry.Geometry )
+		, Tint(FLinearColor(1.0f, 1.0f, 1.0f))
+		, bVisualizeThisNode(true)
+	{ }
+ };
diff --git a/Engine/Source/Developer/SlateReflector/Private/SlateReflectorModule.cpp b/Engine/Source/Developer/SlateReflector/Private/SlateReflectorModule.cpp
new file mode 100644
--- /dev/null
+++ b/Engine/Source/Developer/SlateReflector/Private/SlateReflectorModule.cpp
@@ -0,0 +1,87 @@
+// Copyright 1998-2014 Epic Games, Inc. All Rights Reserved.
+
+/*=============================================================================
+	SlateReflectorModule.cpp: Implements the FSlateReflectorModule class.
+=============================================================================*/
+
+#include "SlateReflectorPrivatePCH.h"
+#include "ModuleManager.h"
+
+
+#define LOCTEXT_NAMESPACE "FSlateReflectorModule"
+
+
+/**
+ * Implements the SlateReflector module.
+ */
+class FSlateReflectorModule
+	: public ISlateReflectorModule
+{
+public:
+
+	// Begin ISlateReflectorModule interface
+
+	virtual TSharedRef<SWidget> GetWidgetReflector( ) OVERRIDE
+	{
+		TSharedPtr<SWidgetReflector> WidgetReflector = WidgetReflectorPtr.Pin();
+
+		if (!WidgetReflector.IsValid())
+		{
+			WidgetReflector = SNew(SWidgetReflector);
+			FSlateApplication::Get().SetWidgetReflector(WidgetReflector.ToSharedRef());
+		}
+
+		return WidgetReflector.ToSharedRef();
+	}
+
+	virtual void RegisterTabSpawner( const TSharedRef<FWorkspaceItem>& WorkspaceGroup ) OVERRIDE
+	{
+		FGlobalTabmanager::Get()->RegisterNomadTabSpawner("WidgetReflector", FOnSpawnTab::CreateRaw(this, &FSlateReflectorModule::MakeWidgetReflectorTab) )
+			.SetDisplayName(LOCTEXT("WidgetReflectorTitle", "Widget Reflector"))
+			.SetTooltipText(LOCTEXT("WidgetReflectorTooltipText", "Open the Widget Reflector tab."))
+			.SetGroup(WorkspaceGroup)
+			.SetIcon(FSlateIcon(FCoreStyle::Get().GetStyleSetName(), "WidgetReflector.TabIcon"));
+	}
+
+	virtual void UnregisterTabSpawner( ) 
+	{
+		FGlobalTabmanager::Get()->UnregisterNomadTabSpawner("WidgetReflector");
+	}
+
+	// End ISlateReflectorModule interface
+
+public:
+
+	// Begin IModuleInterface interface
+
+	virtual void StartupModule( ) OVERRIDE { }
+
+	virtual void ShutdownModule( ) OVERRIDE
+	{
+		UnregisterTabSpawner();
+	}
+
+	// End IModuleInterface interface
+
+private:
+
+	TSharedRef<SDockTab> MakeWidgetReflectorTab( const FSpawnTabArgs& )
+	{
+		return SNew(SDockTab)
+			.TabRole(ETabRole::NomadTab)
+			[
+				GetWidgetReflector()
+			];
+	}
+
+private:
+
+	// Holds the widget reflector singleton.
+	TWeakPtr<SWidgetReflector> WidgetReflectorPtr;
+};
+
+
+IMPLEMENT_MODULE(FSlateReflectorModule, SlateReflector);
+
+
+#undef LOCTEXT_NAMESPACE
diff --git a/Engine/Source/Developer/SlateReflector/Private/SlateReflectorPrivatePCH.h b/Engine/Source/Developer/SlateReflector/Private/SlateReflectorPrivatePCH.h
new file mode 100644
--- /dev/null
+++ b/Engine/Source/Developer/SlateReflector/Private/SlateReflectorPrivatePCH.h
@@ -0,0 +1,26 @@
+// Copyright 1998-2014 Epic Games, Inc. All Rights Reserved.
+
+/*=============================================================================
+	SlateReflectorPrivatePCH.h: Pre-compiled header file for the SlateReflector module.
+=============================================================================*/
+
+#pragma once
+
+#include "SlateReflector.h"
+
+
+/* Dependencies
+ *****************************************************************************/
+
+#include "InputCore.h"
+#include "Slate.h"
+#include "SlateCore.h"
+
+
+/* Private includes
+ *****************************************************************************/
+
+#include "WidgetReflectorNode.h"
+#include "SWidgetReflectorTreeWidgetItem.h"
+#include "SWidgetReflectorToolTipWidget.h"
+#include "SWidgetReflector.h"
diff --git a/Engine/Source/Developer/SlateReflector/Private/Widgets/SWidgetReflector.cpp b/Engine/Source/Developer/SlateReflector/Private/Widgets/SWidgetReflector.cpp
new file mode 100644
--- /dev/null
+++ b/Engine/Source/Developer/SlateReflector/Private/Widgets/SWidgetReflector.cpp
@@ -0,0 +1,321 @@
+// Copyright 1998-2014 Epic Games, Inc. All Rights Reserved.
+
+/*=============================================================================
+	SWidgetReflector.cpp: Implements the SWidgetReflector class.
+=============================================================================*/
+
+#include "SlateReflectorPrivatePCH.h"
+
+
+#define LOCTEXT_NAMESPACE "SWidgetReflector"
+
+
+/* SWidgetReflector interface
+ *****************************************************************************/
+
+void SWidgetReflector::Construct( const FArguments& InArgs )
+{
+	bShowFocus = false;
+	bIsPicking = false;
+
+	this->ChildSlot
+	[
+		SNew(SBorder)
+			.BorderImage(FCoreStyle::Get().GetBrush("ToolPanel.GroupBorder"))
+			[
+				SNew(SVerticalBox)
+				
+				+ SVerticalBox::Slot()
+					.AutoHeight()
+					[
+						SNew(SHorizontalBox)
+					
+						+ SHorizontalBox::Slot()
+							.AutoWidth()
+							[
+								SNew(STextBlock)
+									.Text(LOCTEXT("AppScale", "Application Scale: ").ToString())
+							]
+
+						+ SHorizontalBox::Slot()
+							.AutoWidth()
+							[
+								SNew(SSpinBox<float>)
+									.Value(this, &SWidgetReflector::HandleAppScaleSliderValue)
+									.MinValue(0.1f)
+									.MaxValue(3.0f)
+									.Delta(0.01f)
+									.OnValueChanged(this, &SWidgetReflector::HandleAppScaleSliderChanged)
+							]
+					]
+
+				+ SVerticalBox::Slot()
+					.AutoHeight()
+					[
+						SNew(SHorizontalBox)
+
+						+ SHorizontalBox::Slot()
+							.AutoWidth()
+							.Padding(5.0f)
+							[
+								// Check box that controls LIVE MODE
+								SNew(SCheckBox)
+									.IsChecked(this, &SWidgetReflector::HandleFocusCheckBoxIsChecked)
+									.OnCheckStateChanged(this, &SWidgetReflector::HandleFocusCheckBoxCheckedStateChanged)
+									[
+										SNew(STextBlock)
+											.Text(LOCTEXT("ShowFocus", "Show Focus").ToString())
+									]
+							]
+
+						+ SHorizontalBox::Slot()
+							.AutoWidth()
+							.Padding(5.0f)
+							[
+								// Check box that controls PICKING A WIDGET TO INSPECT
+								SNew(SButton)
+									.OnClicked(this, &SWidgetReflector::HandlePickButtonClicked)
+									.ButtonColorAndOpacity(this, &SWidgetReflector::HandlePickButtonColorAndOpacity)
+									[
+										SNew(STextBlock)
+											.Text(this, &SWidgetReflector::HandlePickButtonText)
+									]
+							]
+					]
+
+				+ SVerticalBox::Slot()
+					.FillHeight(1.0f)
+					[
+						// The tree view that shows all the info that we capture.
+						SAssignNew(ReflectorTree, SReflectorTree)
+							.ItemHeight(24.0f)
+							.TreeItemsSource(&ReflectorTreeRoot)
+							.OnGenerateRow(this, &SWidgetReflector::HandleReflectorTreeGenerateRow)
+							.OnGetChildren(this, &SWidgetReflector::HandleReflectorTreeGetChildren)
+							.OnSelectionChanged(this, &SWidgetReflector::HandleReflectorTreeSelectionChanged)
+							.HeaderRow
+							(
+								SNew(SHeaderRow)
+
+								+ SHeaderRow::Column("WidgetName")
+									.DefaultLabel(LOCTEXT("WidgetName", "Widget Name").ToString())
+									.FillWidth(0.50f)
+
+								+ SHeaderRow::Column("ForegroundColor")
+									.DefaultLabel(LOCTEXT("ForegroundColor", "Foreground Color").ToString())
+									.FixedWidth(20.0f)
+
+								+ SHeaderRow::Column("Visibility")
+									.DefaultLabel(LOCTEXT("Visibility", "Visibility" ).ToString())
+									.FixedWidth(125.0f)
+
+								+ SHeaderRow::Column("WidgetInfo")
+									.DefaultLabel(LOCTEXT("WidgetInfo", "Widget Info" ).ToString())
+									.FillWidth(30.0f)
+							)
+					]
+
+				+ SVerticalBox::Slot()
+					.AutoHeight()
+					[
+						// Frame rate counter
+						SNew(STextBlock)
+							.Text(this, &SWidgetReflector::HandleFrameRateText)
+					]
+			]
+	];
+}
+
+
+/* SCompoundWidget overrides
+ *****************************************************************************/
+
+void SWidgetReflector::Tick( const FGeometry& AllottedGeometry, const double InCurrentTime, const float InDeltaTime )
+{
+	SCompoundWidget::Tick(AllottedGeometry, InCurrentTime, InDeltaTime);
+	ReflectorTree->RequestTreeRefresh();
+}
+
+
+/* IWidgetReflector overrides
+ *****************************************************************************/
+
+bool SWidgetReflector::ReflectorNeedsToDrawIn( TSharedRef<SWindow> ThisWindow ) const
+{
+	return ((SelectedNodes.Num() > 0) && (ReflectorTreeRoot.Num() > 0) && (ReflectorTreeRoot[0]->Widget.Pin() == ThisWindow));
+}
+
+
+void SWidgetReflector::SetWidgetsToVisualize( const FWidgetPath& InWidgetsToVisualize )
+{
+	ReflectorTreeRoot.Empty();
+
+	if (InWidgetsToVisualize.IsValid())
+	{
+		ReflectorTreeRoot.Add(FReflectorNode::NewTreeFrom(InWidgetsToVisualize.Widgets(0)));
+		PickedPath.Empty();
+
+		FReflectorNode::FindWidgetPath( ReflectorTreeRoot, InWidgetsToVisualize, PickedPath );
+		VisualizeAsTree(PickedPath);
+	}
+}
+
+
+int32 SWidgetReflector::Visualize( const FWidgetPath& InWidgetsToVisualize, FSlateWindowElementList& OutDrawElements, int32 LayerId )
+{
+	const bool bAttemptingToVisualizeReflector = InWidgetsToVisualize.ContainsWidget(SharedThis(this));
+
+	if (!InWidgetsToVisualize.IsValid())
+	{
+		TSharedPtr<SWidget> WindowWidget = ReflectorTreeRoot[0]->Widget.Pin();
+		TSharedPtr<SWindow> Window = StaticCastSharedPtr<SWindow>(WindowWidget);
+
+		return VisualizeSelectedNodesAsRectangles(SelectedNodes, Window.ToSharedRef(), OutDrawElements, LayerId);
+	}
+
+	if (!bAttemptingToVisualizeReflector)
+	{
+		SetWidgetsToVisualize(InWidgetsToVisualize);
+
+		return VisualizePickAsRectangles(InWidgetsToVisualize, OutDrawElements, LayerId);
+	}		
+
+	return LayerId;
+}
+
+
+/* SWidgetReflector implementation
+ *****************************************************************************/
+
+TSharedRef<SToolTip> SWidgetReflector::GenerateToolTipForReflectorNode( TSharedPtr<FReflectorNode> InReflectorNode )
+{
+	return SNew(SToolTip)
+	[
+		SNew(SReflectorToolTipWidget)
+			.WidgetInfoToVisualize(InReflectorNode)
+	];
+}
+
+
+void SWidgetReflector::VisualizeAsTree( const TArray< TSharedPtr<FReflectorNode> >& WidgetPathToVisualize )
+{
+	const FLinearColor TopmostWidgetColor(1.0f, 0.0f, 0.0f);
+	const FLinearColor LeafmostWidgetColor(0.0f, 1.0f, 0.0f);
+
+	for (int32 WidgetIndex = 0; WidgetIndex<WidgetPathToVisualize.Num(); ++WidgetIndex)
+	{
+		TSharedPtr<FReflectorNode> CurWidget = WidgetPathToVisualize[WidgetIndex];
+
+		// Tint the item based on depth in picked path
+		const float ColorFactor = static_cast<float>(WidgetIndex)/WidgetPathToVisualize.Num();
+		CurWidget->Tint = FMath::Lerp(TopmostWidgetColor, LeafmostWidgetColor, ColorFactor);
+
+		// Make sure the user can see the picked path in the tree.
+		ReflectorTree->SetItemExpansion(CurWidget, true);
+	}
+
+	ReflectorTree->RequestScrollIntoView(WidgetPathToVisualize.Last());
+	ReflectorTree->SetSelection(WidgetPathToVisualize.Last());
+}
+
+
+int32 SWidgetReflector::VisualizePickAsRectangles( const FWidgetPath& InWidgetsToVisualize, FSlateWindowElementList& OutDrawElements, int32 LayerId)
+{
+	const FLinearColor TopmostWidgetColor(1.0f, 0.0f, 0.0f);
+	const FLinearColor LeafmostWidgetColor(0.0f, 1.0f, 0.0f);
+
+	for (int32 WidgetIndex = 0; WidgetIndex < InWidgetsToVisualize.Widgets.Num(); ++WidgetIndex)
+	{
+		const FArrangedWidget& WidgetGeometry = InWidgetsToVisualize.Widgets(WidgetIndex);
+		const float ColorFactor = static_cast<float>(WidgetIndex)/InWidgetsToVisualize.Widgets.Num();
+		const FLinearColor Tint(1.0f - ColorFactor, ColorFactor, 0.0f, 1.0f);
+
+		FSlateDrawElement::MakeBox(
+			OutDrawElements,
+			++LayerId,
+			FPaintGeometry(
+				WidgetGeometry.Geometry.AbsolutePosition - InWidgetsToVisualize.TopLevelWindow->GetPositionInScreen(),
+				WidgetGeometry.Geometry.Size * WidgetGeometry.Geometry.Scale,
+				WidgetGeometry.Geometry.Scale),
+				FCoreStyle::Get().GetBrush(TEXT("Debug.Border")),
+				InWidgetsToVisualize.TopLevelWindow->GetClippingRectangleInWindow(),
+				ESlateDrawEffect::None,
+				FMath::Lerp( TopmostWidgetColor, LeafmostWidgetColor, ColorFactor
+			)
+		);
+	}
+
+	return LayerId;
+}
+
+
+int32 SWidgetReflector::VisualizeSelectedNodesAsRectangles( const TArray<TSharedPtr<FReflectorNode>>& InNodesToDraw, const TSharedRef<SWindow>& VisualizeInWindow, FSlateWindowElementList& OutDrawElements, int32 LayerId )
+{
+	for (int32 NodeIndex = 0; NodeIndex < InNodesToDraw.Num(); ++NodeIndex)
+	{
+		const TSharedPtr<FReflectorNode>& NodeToDraw = InNodesToDraw[NodeIndex];
+		const FLinearColor Tint(0.0f, 1.0f, 0.0f);
+
+		FSlateDrawElement::MakeBox(
+			OutDrawElements,
+			++LayerId,
+			FPaintGeometry(
+				NodeToDraw->Geometry.AbsolutePosition - VisualizeInWindow->GetPositionInScreen(),
+				NodeToDraw->Geometry.Size * NodeToDraw->Geometry.Scale,
+				NodeToDraw->Geometry.Scale
+			),
+			FCoreStyle::Get().GetBrush(TEXT("Debug.Border")),
+			VisualizeInWindow->GetClippingRectangleInWindow(),
+			ESlateDrawEffect::None,
+			NodeToDraw->Tint
+		);
+	}
+
+	return LayerId;
+}
+
+
+/* SWidgetReflector callbacks
+ *****************************************************************************/
+
+void SWidgetReflector::HandleFocusCheckBoxCheckedStateChanged( ESlateCheckBoxState::Type NewValue )
+{
+	bShowFocus = NewValue != ESlateCheckBoxState::Unchecked;
+
+	if (bShowFocus)
+	{
+		bIsPicking = false;
+	}
+}
+
+
+FString SWidgetReflector::HandleFrameRateText( ) const
+{
+	FString MyString;
+#if 0 // the new stats system does not support this
+	MyString = FString::Printf(TEXT("FPS: %0.2f (%0.2f ms)"), (float)( 1.0f / FPSCounter.GetAverage()), (float)FPSCounter.GetAverage() * 1000.0f);
+#endif
+
+	return MyString;
+}
+
+
+FString SWidgetReflector::HandlePickButtonText( ) const
+{
+	static const FString NotPicking = LOCTEXT("PickWidget", "Pick Widget").ToString();
+	static const FString Picking = LOCTEXT("PickingWidget", "Picking (Esc to Stop)").ToString();
+
+	return bIsPicking ? Picking : NotPicking;
+}
+
+
+TSharedRef<ITableRow> SWidgetReflector::HandleReflectorTreeGenerateRow( TSharedPtr<FReflectorNode> InReflectorNode, const TSharedRef<STableViewBase>& OwnerTable )
+{
+	return SNew(SReflectorTreeWidgetItem, OwnerTable)
+		.WidgetInfoToVisualize(InReflectorNode)
+		.ToolTip(GenerateToolTipForReflectorNode(InReflectorNode))
+		.SourceCodeAccessor(SourceAccessDelegate);
+}
+
+
+#undef LOCTEXT_NAMESPACE
diff --git a/Engine/Source/Developer/SlateReflector/Private/Widgets/SWidgetReflector.h b/Engine/Source/Developer/SlateReflector/Private/Widgets/SWidgetReflector.h
new file mode 100644
--- /dev/null
+++ b/Engine/Source/Developer/SlateReflector/Private/Widgets/SWidgetReflector.h
@@ -0,0 +1,202 @@
+// Copyright 1998-2014 Epic Games, Inc. All Rights Reserved.
+
+/*=============================================================================
+	SWidgetReflector.h: Declares the SWidgetReflector class.
+=============================================================================*/
+
+#pragma once
+
+
+/** A Delegate for passing along a string of a source code location to access */
+DECLARE_DELEGATE_RetVal_ThreeParams(bool, FAccessSourceCode, const FString& /*FileName*/, int32 /*InLineNumber*/, int32 /*InColumnNumber*/);
+
+
+/**
+ * Abstract base class for widget reflectors.
+ */
+class SWidgetReflector
+	: public SCompoundWidget
+	, public IWidgetReflector
+{
+	// The reflector uses a tree that observes FReflectorNodes.
+	typedef STreeView<TSharedPtr<FReflectorNode>> SReflectorTree;
+
+public:
+
+	SLATE_BEGIN_ARGS(SWidgetReflector) { }
+	SLATE_END_ARGS()
+
+	/** Create widgets that comprise this WidgetReflector implementation */
+	/**
+	 * Creates and initializes a new widget reflector widget.
+	 *
+	 * @param InArgs The construction arguments.
+	 */
+	void Construct( const FArguments& InArgs );
+
+public:
+
+	// Begin SCompoundWidget overrides
+
+	virtual void Tick( const FGeometry& AllottedGeometry, const double InCurrentTime, const float InDeltaTime ) OVERRIDE;
+
+	// End SCompoundWidget overrides
+
+public:
+
+	// Begin IWidgetReflector interface
+
+	virtual bool IsInPickingMode( ) const OVERRIDE
+	{
+		return bIsPicking;
+	}
+
+	virtual bool IsShowingFocus( ) const OVERRIDE
+	{
+		return bShowFocus;
+	}
+
+	virtual bool IsVisualizingLayoutUnderCursor( ) const OVERRIDE
+	{
+		return bIsPicking;
+	}
+
+	virtual void OnWidgetPicked( ) OVERRIDE
+	{
+		bIsPicking = false;
+	}
+
+	virtual bool ReflectorNeedsToDrawIn( TSharedRef<SWindow> ThisWindow ) const OVERRIDE;
+
+	virtual void SetSourceAccessDelegate( FAccessSourceCode InDelegate ) OVERRIDE
+	{
+		SourceAccessDelegate = InDelegate;
+	}
+
+	virtual void SetWidgetsToVisualize( const FWidgetPath& InWidgetsToVisualize ) OVERRIDE;
+
+	virtual int32 Visualize( const FWidgetPath& InWidgetsToVisualize, FSlateWindowElementList& OutDrawElements, int32 LayerId ) OVERRIDE;
+
+	// End IWidgetReflector interface
+
+protected:
+
+	/**
+	 * Generates a tool tip for the given reflector tree node.
+	 *
+	 * @param InReflectorNode The node to generate the tool tip for.
+	 *
+	 * @return The tool tip widget.
+	 */
+	TSharedRef<SToolTip> GenerateToolTipForReflectorNode( TSharedPtr<FReflectorNode> InReflectorNode );
+
+	/**
+	 * Mark the provided reflector nodes such that they stand out in the tree and are visible.
+	 *
+	 * @param WidgetPathToObserve The nodes to mark.
+	 */
+	void VisualizeAsTree( const TArray< TSharedPtr<FReflectorNode> >& WidgetPathToVisualize );
+
+	/**
+	 * Draw the widget path to the picked widget as the widgets' outlines.
+	 *
+	 * @param InWidgetsToVisualize A widget path whose widgets' outlines to draw.
+	 * @param OutDrawElements A list of draw elements; we will add the output outlines into it.
+	 * @param LayerId The maximum layer achieved in OutDrawElements so far.
+	 *
+	 * @return The maximum layer ID we achieved while painting.
+	 */
+	int32 VisualizePickAsRectangles( const FWidgetPath& InWidgetsToVisualize, FSlateWindowElementList& OutDrawElements, int32 LayerId );
+
+	/**
+	 * Draw an outline for the specified nodes.
+	 *
+	 * @param InNodesToDraw A widget path whose widgets' outlines to draw.
+	 * @param WindowGeometry The geometry of the window in which to draw.
+	 * @param OutDrawElements A list of draw elements; we will add the output outlines into it.
+	 * @param LayerId the maximum layer achieved in OutDrawElements so far.
+	 *
+	 * @return The maximum layer ID we achieved while painting.
+	 */
+	int32 VisualizeSelectedNodesAsRectangles( const TArray<TSharedPtr<FReflectorNode>>& InNodesToDraw, const TSharedRef<SWindow>& VisualizeInWindow, FSlateWindowElementList& OutDrawElements, int32 LayerId );
+
+private:
+
+	// Callback for changing the application scale slider
+	void HandleAppScaleSliderChanged( float NewValue )
+	{
+		FSlateApplication::Get().SetApplicationScale(NewValue);
+	}
+
+	// Callback for getting the value of the application scale slider.
+	float HandleAppScaleSliderValue( ) const
+	{
+		return FSlateApplication::Get().GetApplicationScale();
+	}
+
+	// Callback for checked state changes of the focus check box.
+	void HandleFocusCheckBoxCheckedStateChanged( ESlateCheckBoxState::Type NewValue );
+
+	// Callback for getting the checked state of the focus check box.
+	ESlateCheckBoxState::Type HandleFocusCheckBoxIsChecked( ) const
+	{
+		return bShowFocus ? ESlateCheckBoxState::Checked : ESlateCheckBoxState::Unchecked;
+	}
+
+	// Callback for getting the text of the frame rate text block.
+	FString HandleFrameRateText( ) const;
+
+	// Callback for clicking the pick button.
+	FReply HandlePickButtonClicked( )
+	{
+		bIsPicking = !bIsPicking;
+
+		if (bIsPicking)
+		{
+			bShowFocus = false;
+		}
+
+		return FReply::Handled();
+	}
+
+	// Callback for getting the color of the pick button text.
+	FSlateColor HandlePickButtonColorAndOpacity( ) const
+	{
+		return bIsPicking
+			? FCoreStyle::Get().GetSlateColor("SelectionColor")
+			: FCoreStyle::Get().GetSlateColor("DefaultForeground");
+	}
+
+	// Callback for getting the text of the pick button.
+	FString HandlePickButtonText() const;
+
+	// Callback for generating a row in the reflector tree view.
+	TSharedRef<ITableRow> HandleReflectorTreeGenerateRow( TSharedPtr<FReflectorNode> InReflectorNode, const TSharedRef<STableViewBase>& OwnerTable );
+
+	// Callback for getting the child items of the given reflector tree node.
+	void HandleReflectorTreeGetChildren( TSharedPtr<FReflectorNode> InWidgetGeometry, TArray<TSharedPtr<FReflectorNode>>& OutChildren )
+	{
+		OutChildren = InWidgetGeometry->ChildNodes;
+	}
+
+	// Callback for when the selection in the reflector tree has changed.
+	void HandleReflectorTreeSelectionChanged( TSharedPtr< FReflectorNode >, ESelectInfo::Type /*SelectInfo*/ )
+	{
+		SelectedNodes = ReflectorTree->GetSelectedItems();
+	}
+
+private:
+
+	TSharedPtr<SReflectorTree> ReflectorTree;
+
+	TArray<TSharedPtr<FReflectorNode>> SelectedNodes;
+	TArray<TSharedPtr<FReflectorNode>> ReflectorTreeRoot;
+	TArray<TSharedPtr<FReflectorNode>> PickedPath;
+
+	SSplitter::FSlot* WidgetInfoLocation;
+
+	FAccessSourceCode SourceAccessDelegate;
+
+	bool bShowFocus;
+	bool bIsPicking;
+};
diff --git a/Engine/Source/Developer/SlateReflector/Private/Widgets/SWidgetReflectorToolTipWidget.h b/Engine/Source/Developer/SlateReflector/Private/Widgets/SWidgetReflectorToolTipWidget.h
new file mode 100644
--- /dev/null
+++ b/Engine/Source/Developer/SlateReflector/Private/Widgets/SWidgetReflectorToolTipWidget.h
@@ -0,0 +1,139 @@
+// Copyright 1998-2014 Epic Games, Inc. All Rights Reserved.
+
+/*=============================================================================
+	SReflectorToolTipWidget.h: Declares the SReflectorToolTipWidget class.
+=============================================================================*/
+
+#pragma once
+
+
+#define LOCTEXT_NAMESPACE "SReflectorToolTipWidget"
+
+
+class SReflectorToolTipWidget
+	: public SCompoundWidget
+{
+public:
+
+	SLATE_BEGIN_ARGS(SReflectorToolTipWidget)
+		: _WidgetInfoToVisualize()
+	{ }
+
+		SLATE_ARGUMENT(TSharedPtr<FReflectorNode>, WidgetInfoToVisualize)
+
+	SLATE_END_ARGS()
+
+	/** Constructs this widget with InArgs */
+	void Construct( const FArguments& InArgs )
+	{
+		this->WidgetInfo = InArgs._WidgetInfoToVisualize;
+
+		ChildSlot
+		[
+			SNew(SGridPanel)
+				.FillColumn(1, 1.0f)
+
+			// Desired Size
+			+ SGridPanel::Slot(0, 0)
+			[
+				SNew(STextBlock)
+					.Text(LOCTEXT("DesiredSize", "Desired Size"))
+			]
+
+			// Desired Size Value
+			+ SGridPanel::Slot(1, 0)
+				.Padding(5.0f, 0.0f, 0.0f, 0.0f)
+				[
+					SNew(STextBlock)
+						.Text(this, &SReflectorToolTipWidget::GetWidgetsDesiredSize)
+				]
+
+			// Actual Size
+			+ SGridPanel::Slot(0, 1)
+				[
+					SNew( STextBlock )
+						.Text(LOCTEXT("ActualSize", "Actual Size"))
+				]
+
+			// Actual Size Value
+			+ SGridPanel::Slot(1, 1)
+				.Padding(5.0f, 0.0f, 0.0f, 0.0f)
+				[
+					SNew(STextBlock)
+						.Text(this, &SReflectorToolTipWidget::GetWidgetActualSize)
+				]
+
+			// Size Info
+			+ SGridPanel::Slot(0, 2)
+				[
+					SNew(STextBlock)
+						.Text(LOCTEXT("SizeInfo", "Size Info"))
+				]
+
+			// Size Info Value
+			+ SGridPanel::Slot(1, 2)
+				.Padding(5.0f, 0.0f, 0.0f, 0.0f)
+				[
+					SNew(STextBlock)
+						.Text(this, &SReflectorToolTipWidget::GetSizeInfo)
+				]
+
+			// Enabled
+			+ SGridPanel::Slot(0, 3)
+				[
+					SNew(STextBlock)
+						.Text(LOCTEXT("Enabled", "Enabled"))
+				]
+
+			// Enabled Value
+			+ SGridPanel::Slot(1, 3)
+				.Padding(5.0f, 0.0f, 0.0f, 0.0f)
+				[
+					SNew(STextBlock)
+						.Text(this, &SReflectorToolTipWidget::GetEnabled)
+				]
+		];
+	}
+
+private:
+
+	FString GetWidgetsDesiredSize( ) const
+	{
+		TSharedPtr<SWidget> TheWidget = WidgetInfo.Get()->Widget.Pin();
+
+		return TheWidget.IsValid()
+			? TheWidget->GetDesiredSize().ToString()
+			: FString();
+	}
+
+	FString GetWidgetActualSize( ) const
+	{
+		return WidgetInfo.Get()->Geometry.Size.ToString();
+	}
+
+	FString GetSizeInfo( ) const
+	{
+		TSharedPtr<SWidget> TheWidget = WidgetInfo.Get()->Widget.Pin();
+
+		return TheWidget.IsValid()
+			? WidgetInfo.Get()->Geometry.ToString()
+			: FString();
+	}
+
+	FString GetEnabled( ) const
+	{
+		TSharedPtr<SWidget> TheWidget = WidgetInfo.Get()->Widget.Pin();
+
+		return TheWidget.IsValid()
+			? (TheWidget->IsEnabled() ? "True" : "False")
+			: FString();
+	}
+
+private:
+
+	// The info about the widget that we are visualizing.
+	TAttribute<TSharedPtr<FReflectorNode>> WidgetInfo;
+};
+
+
+#undef LOCTEXT_NAMESPACE
diff --git a/Engine/Source/Developer/SlateReflector/Private/Widgets/SWidgetReflectorTreeWidgetItem.h b/Engine/Source/Developer/SlateReflector/Private/Widgets/SWidgetReflectorTreeWidgetItem.h
new file mode 100644
--- /dev/null
+++ b/Engine/Source/Developer/SlateReflector/Private/Widgets/SWidgetReflectorTreeWidgetItem.h
@@ -0,0 +1,175 @@
+// Copyright 1998-2014 Epic Games, Inc. All Rights Reserved.
+
+/*=============================================================================
+	SReflectorTreeWidgetItem.h: Declares the SReflectorTreeWidgetItem class.
+=============================================================================*/
+
+#pragma once
+
+
+/**
+ * Widget that visualizes the contents of a FReflectorNode
+ */
+class SReflectorTreeWidgetItem
+	: public SMultiColumnTableRow<TSharedPtr<FReflectorNode>>
+{
+public:
+
+	SLATE_BEGIN_ARGS(SReflectorTreeWidgetItem)
+		: _WidgetInfoToVisualize()
+		, _SourceCodeAccessor()
+	{ }
+
+		SLATE_ARGUMENT(TSharedPtr<FReflectorNode>, WidgetInfoToVisualize)
+		SLATE_ARGUMENT(FAccessSourceCode, SourceCodeAccessor)
+
+	SLATE_END_ARGS()
+
+public:
+
+	/**
+	 * Construct child widgets that comprise this widget.
+	 *
+	 * @param InArgs  Declaration from which to construct this widget
+	 */
+	void Construct( const FArguments& InArgs, const TSharedRef<STableViewBase>& InOwnerTableView )
+	{
+		this->WidgetInfo = InArgs._WidgetInfoToVisualize;
+		this->OnAccessSourceCode = InArgs._SourceCodeAccessor;
+
+		SMultiColumnTableRow< TSharedPtr<FReflectorNode> >::Construct( SMultiColumnTableRow< TSharedPtr<FReflectorNode> >::FArguments().Padding(1), InOwnerTableView );
+	}
+
+	BEGIN_SLATE_FUNCTION_BUILD_OPTIMIZATION
+	virtual TSharedRef<SWidget> GenerateWidgetForColumn( const FName& ColumnName ) OVERRIDE
+	{
+		if (ColumnName == TEXT("WidgetName"))
+		{
+			return SNew(SHorizontalBox)
+			
+			+ SHorizontalBox::Slot()
+				.AutoWidth()
+				.VAlign(VAlign_Center)
+				[
+					SNew(SExpanderArrow, SharedThis(this))
+				]
+
+			+ SHorizontalBox::Slot()
+				.AutoWidth()
+				.Padding(2.0f, 0.0f)
+				.VAlign(VAlign_Center)
+				[
+					SNew(STextBlock)
+						.Text(this, &SReflectorTreeWidgetItem::GetWidgetType)
+						.ColorAndOpacity(this, &SReflectorTreeWidgetItem::GetTint)
+				];
+		}
+		else if (ColumnName == TEXT("WidgetInfo"))
+		{
+			return SNew(SBox)
+				.HAlign(HAlign_Left)
+				.VAlign(VAlign_Center)
+				.Padding(FMargin(2.0f, 0.0f))
+				[
+					SNew(SHyperlink)
+						.Text(this, &SReflectorTreeWidgetItem::GetReadableLocation)
+						.OnNavigate(this, &SReflectorTreeWidgetItem::HandleHyperlinkNavigate)
+				];
+		}
+		else if (ColumnName == "Visibility")
+		{
+			return SNew(SBox)
+				.HAlign(HAlign_Center)
+				.VAlign(VAlign_Center)
+				.Padding(FMargin(2.0f, 0.0f))
+				[
+					SNew(STextBlock)
+						.Text(this, &SReflectorTreeWidgetItem::GetVisibilityAsString)
+				];
+		}
+		else if (ColumnName == "ForegroundColor")
+		{
+			TSharedPtr<SWidget> ThisWidget = WidgetInfo.Get()->Widget.Pin();
+			FSlateColor Foreground = (ThisWidget.IsValid())
+				? ThisWidget->GetForegroundColor()
+				: FSlateColor::UseForeground();
+
+			return SNew(SBorder)
+				// Show unset color as an empty space.
+				.Visibility(Foreground.IsColorSpecified() ? EVisibility::Visible : EVisibility::Hidden)
+				// Show a checkerboard background so we can see alpha values well
+				.BorderImage( FCoreStyle::Get().GetBrush("Checkerboard") )
+				.VAlign(VAlign_Center)
+				.Padding(FMargin(2.0f, 0.0f))
+				[
+					// Show a color block
+					SNew(SColorBlock)
+						.Color(Foreground.GetSpecifiedColor())
+						.Size(FVector2D(16.0f, 16.0f))
+				];
+		}
+		else
+		{
+			return SNullWidget::NullWidget;
+		}
+	}
+	END_SLATE_FUNCTION_BUILD_OPTIMIZATION
+
+protected:
+
+	/** @return String representation of the widget we are visualizing */
+	FString GetWidgetType() const
+	{
+		return WidgetInfo.Get()->Widget.IsValid()
+			? WidgetInfo.Get()->Widget.Pin()->GetTypeAsString()
+			: TEXT("Null Widget");
+	}
+	
+	FString GetReadableLocation() const
+	{
+		return WidgetInfo.Get()->Widget.IsValid()
+			? WidgetInfo.Get()->Widget.Pin()->GetReadableLocation()
+			: FString();
+	}
+
+	FString GetWidgetFile() const
+	{
+		return WidgetInfo.Get()->Widget.IsValid()
+			? WidgetInfo.Get()->Widget.Pin()->GetCreatedInFile()
+			: FString();
+	}
+
+	int32 GetWidgetLineNumber() const
+	{
+		return WidgetInfo.Get()->Widget.IsValid()
+			? WidgetInfo.Get()->Widget.Pin()->GetCreatedInLineNumber()
+			: 0;
+	}
+
+	FString GetVisibilityAsString() const
+	{
+		TSharedPtr<SWidget> TheWidget = WidgetInfo.Get()->Widget.Pin();
+		return TheWidget.IsValid()
+			? TheWidget->GetVisibility().ToString()
+			: FString();
+	}
+
+	/** @return The tint of the reflector node */
+	FSlateColor GetTint () const
+	{
+		return WidgetInfo.Get()->Tint;
+	}
+
+	void HandleHyperlinkNavigate()
+	{
+		if(OnAccessSourceCode.IsBound())
+		{
+			OnAccessSourceCode.Execute(GetWidgetFile(), GetWidgetLineNumber(), 0);
+		}
+	}
+
+	/** The info about the widget that we are visualizing */
+	TAttribute< TSharedPtr<FReflectorNode> > WidgetInfo;
+
+	FAccessSourceCode OnAccessSourceCode;
+};
diff --git a/Engine/Source/Developer/SlateReflector/Public/Interfaces/ISlateReflectorModule.h b/Engine/Source/Developer/SlateReflector/Public/Interfaces/ISlateReflectorModule.h
new file mode 100644
--- /dev/null
+++ b/Engine/Source/Developer/SlateReflector/Public/Interfaces/ISlateReflectorModule.h
@@ -0,0 +1,43 @@
+// Copyright 1998-2014 Epic Games, Inc. All Rights Reserved.
+
+/*=============================================================================
+	IMessagingModule.h: Declares the IMessagingModule interface.
+=============================================================================*/
+
+#pragma once
+
+
+/**
+ * Interface for messaging modules.
+ */
+class ISlateReflectorModule
+	: public IModuleInterface
+{
+public:
+
+	/**
+	 * Creates a widget reflector widget.
+	 *
+	 * @return The widget.
+	 */
+	virtual TSharedRef<SWidget> GetWidgetReflector( ) = 0;
+
+	/**
+	 * Registers a tab spawner for the widget reflector.
+	 *
+	 * @param WorkspaceGroup The workspace group to insert the tab into.
+	 */
+	virtual void RegisterTabSpawner( const TSharedRef<FWorkspaceItem>& WorkspaceGroup ) = 0;
+
+	/**
+	 * Unregisters the tab spawner for the widget reflector.
+	 */
+	virtual void UnregisterTabSpawner( ) = 0;
+
+public:
+
+	/**
+	 * Virtual destructor.
+	 */
+	virtual ~ISlateReflectorModule( ) { }
+};
diff --git a/Engine/Source/Developer/SlateReflector/Public/SlateReflector.h b/Engine/Source/Developer/SlateReflector/Public/SlateReflector.h
new file mode 100644
--- /dev/null
+++ b/Engine/Source/Developer/SlateReflector/Public/SlateReflector.h
@@ -0,0 +1,21 @@
+// Copyright 1998-2014 Epic Games, Inc. All Rights Reserved.
+
+/*=============================================================================
+	SlateReflector.h: SlateReflector module public header file.
+=============================================================================*/
+
+#pragma once
+
+
+/* Dependencies
+ *****************************************************************************/
+
+#include "Core.h"
+#include "Slate.h"
+#include "SlateCore.h"
+
+
+/* Interfaces
+ *****************************************************************************/
+
+#include "ISlateReflectorModule.h"
diff --git a/Engine/Source/Developer/SlateReflector/SlateReflector.Build.cs b/Engine/Source/Developer/SlateReflector/SlateReflector.Build.cs
new file mode 100644
--- /dev/null
+++ b/Engine/Source/Developer/SlateReflector/SlateReflector.Build.cs
@@ -0,0 +1,31 @@
+// Copyright 1998-2014 Epic Games, Inc. All Rights Reserved.
+
+using UnrealBuildTool;
+
+public class SlateReflector : ModuleRules
+{
+	public SlateReflector(TargetInfo Target)
+	{
+		PublicDependencyModuleNames.AddRange(
+			new string[] {
+				"Core",
+				"Slate",
+				"SlateCore",
+			}
+		);
+
+		PrivateDependencyModuleNames.AddRange(
+			new string[] {
+				"InputCore",
+			}
+		);
+
+		PrivateIncludePaths.AddRange(
+			new string[] {
+				"Developer/SlateReflector/Private",
+				"Developer/SlateReflector/Private/Models",
+				"Developer/SlateReflector/Private/Widgets",
+			}
+		);
+	}
+}
diff --git a/Engine/Source/Developer/SourceCodeAccess/SourceCodeAccess.Build.cs b/Engine/Source/Developer/SourceCodeAccess/SourceCodeAccess.Build.cs
--- a/Engine/Source/Developer/SourceCodeAccess/SourceCodeAccess.Build.cs
+++ b/Engine/Source/Developer/SourceCodeAccess/SourceCodeAccess.Build.cs
@@ -6,15 +6,20 @@ public class SourceCodeAccess : ModuleRules
 {
 	public SourceCodeAccess(TargetInfo Target)
 	{
-        PrivateIncludePaths.Add("Developer/SourceCodeAccess/Private");
-
 		PublicDependencyModuleNames.AddRange(
 			new string[] {
 				"Core",
-				"CoreUObject",
 				"Settings",
-				"Slate"
 			}
 			);
+
+		PrivateDependencyModuleNames.AddRange(
+			new string[] {
+				"CoreUObject",
+				"Slate",
+			}
+		);
+
+		PrivateIncludePaths.Add("Developer/SourceCodeAccess/Private");
 	}
 }
diff --git a/Engine/Source/Developer/StandaloneRenderer/Private/OpenGL/SlateOpenGLRenderer.h b/Engine/Source/Developer/StandaloneRenderer/Private/OpenGL/SlateOpenGLRenderer.h
--- a/Engine/Source/Developer/StandaloneRenderer/Private/OpenGL/SlateOpenGLRenderer.h
+++ b/Engine/Source/Developer/StandaloneRenderer/Private/OpenGL/SlateOpenGLRenderer.h
@@ -2,7 +2,6 @@
 
 #pragma once
 
-#include "Renderer.h"
 
 class FSlateOpenGLTextureManager;
 class FSlateOpenGLRenderingPolicy;
diff --git a/Engine/Source/Developer/StandaloneRenderer/Private/StandaloneRendererPrivate.h b/Engine/Source/Developer/StandaloneRenderer/Private/StandaloneRendererPrivate.h
--- a/Engine/Source/Developer/StandaloneRenderer/Private/StandaloneRendererPrivate.h
+++ b/Engine/Source/Developer/StandaloneRenderer/Private/StandaloneRendererPrivate.h
@@ -5,7 +5,7 @@
 
 
 #include "Core.h"
-#include "Slate.h"
+#include "SlateCore.h"
 
 #if PLATFORM_WINDOWS
 #include "AllowWindowsPlatformTypes.h"
@@ -32,12 +32,6 @@
 #include "HideWindowsPlatformTypes.h"
 #endif		// PLATFORM_WINDOWS
 
-
-#include "RenderingCommon.h"
-#include "Renderer.h"
-#include "RenderingPolicy.h"
-#include "TextureManager.h"
-
 #if PLATFORM_WINDOWS
 #include "Windows/D3D/SlateD3DTextures.h"
 #include "Windows/D3D/SlateD3DConstantBuffer.h"
diff --git a/Engine/Source/Developer/StandaloneRenderer/Private/Windows/D3D/SlateD3DRenderer.h b/Engine/Source/Developer/StandaloneRenderer/Private/Windows/D3D/SlateD3DRenderer.h
--- a/Engine/Source/Developer/StandaloneRenderer/Private/Windows/D3D/SlateD3DRenderer.h
+++ b/Engine/Source/Developer/StandaloneRenderer/Private/Windows/D3D/SlateD3DRenderer.h
@@ -4,7 +4,6 @@
 #ifndef _SLATERENDERRD3D_H__
 #define _SLATERENDERRD3D_H__
 
-#include "Renderer.h"
 
 class FSlateD3DTextureManager;
 class FSlateD3D11RenderingPolicy;
diff --git a/Engine/Source/Developer/StandaloneRenderer/StandaloneRenderer.Build.cs b/Engine/Source/Developer/StandaloneRenderer/StandaloneRenderer.Build.cs
--- a/Engine/Source/Developer/StandaloneRenderer/StandaloneRenderer.Build.cs
+++ b/Engine/Source/Developer/StandaloneRenderer/StandaloneRenderer.Build.cs
@@ -14,7 +14,7 @@ public class StandaloneRenderer : ModuleRules
 				"CoreUObject",
 				"ImageWrapper",
 				"InputCore",
-				"Slate",
+				"SlateCore",
 			}
 			);
 
diff --git a/Engine/Source/Editor/AnimGraph/AnimGraph.Build.cs b/Engine/Source/Editor/AnimGraph/AnimGraph.Build.cs
--- a/Engine/Source/Editor/AnimGraph/AnimGraph.Build.cs
+++ b/Engine/Source/Editor/AnimGraph/AnimGraph.Build.cs
@@ -15,22 +15,23 @@ public class AnimGraph : ModuleRules
 				"CoreUObject", 
 				"Engine", 
 				"Slate",
-				"BlueprintGraph"
+				"BlueprintGraph",
 			}
-			);
+		);
 
 		PrivateDependencyModuleNames.AddRange(
 			new string[] {
+				"SlateCore",
 				"UnrealEd",
                 "GraphEditor",
 			}
-			);
+		);
 
 		CircularlyReferencedDependentModules.AddRange(
 			new string[] {
                 "UnrealEd",
                 "GraphEditor",
             }
-			); 
+		); 
 	}
 }
diff --git a/Engine/Source/Editor/BlueprintGraph/BlueprintGraph.Build.cs b/Engine/Source/Editor/BlueprintGraph/BlueprintGraph.Build.cs
--- a/Engine/Source/Editor/BlueprintGraph/BlueprintGraph.Build.cs
+++ b/Engine/Source/Editor/BlueprintGraph/BlueprintGraph.Build.cs
@@ -12,7 +12,7 @@ public class BlueprintGraph : ModuleRules
                 "Editor/BlueprintGraph/Private",
                 "Editor/KismetCompiler/Public",
             }
-            );
+		);
 
 		PublicDependencyModuleNames.AddRange(
 			new string[] { 
@@ -23,7 +23,7 @@ public class BlueprintGraph : ModuleRules
 				"Slate",
                 "EditorStyle",
 			}
-			);
+		);
 
 		PrivateDependencyModuleNames.AddRange( 
 			new string[] {
@@ -31,8 +31,9 @@ public class BlueprintGraph : ModuleRules
                 "KismetCompiler",
 				"UnrealEd",
                 "GraphEditor",
+				"SlateCore",
 			}
-			);
+		);
 
 		CircularlyReferencedDependentModules.AddRange(
             new string[] {
@@ -40,7 +41,6 @@ public class BlueprintGraph : ModuleRules
                 "UnrealEd",
                 "GraphEditor",
             }
-            ); 
-
+		); 
 	}
 }
diff --git a/Engine/Source/Editor/BspMode/BspMode.Build.cs b/Engine/Source/Editor/BspMode/BspMode.Build.cs
--- a/Engine/Source/Editor/BspMode/BspMode.Build.cs
+++ b/Engine/Source/Editor/BspMode/BspMode.Build.cs
@@ -13,16 +13,17 @@ public class BspMode : ModuleRules
 				"Engine",
                 "InputCore",
 				"Slate",
+				"SlateCore",
                 "EditorStyle",
 				"UnrealEd",
-				"LevelEditor"
+				"LevelEditor",
 			}
-			);
+		);
 
         DynamicallyLoadedModuleNames.AddRange(
             new string[] {
 				"PropertyEditor",
-                "LevelEditor"
+                "LevelEditor",
 			}
         );
 	}
diff --git a/Engine/Source/Editor/Cascade/Cascade.Build.cs b/Engine/Source/Editor/Cascade/Cascade.Build.cs
--- a/Engine/Source/Editor/Cascade/Cascade.Build.cs
+++ b/Engine/Source/Editor/Cascade/Cascade.Build.cs
@@ -11,7 +11,7 @@ public class Cascade : ModuleRules
 				"Editor/DistCurveEditor/Public",
 				"Editor/UnrealEd/Public",
 			}
-			);
+		);
 
 		PrivateDependencyModuleNames.AddRange(
 			new string[] {
@@ -20,17 +20,18 @@ public class Cascade : ModuleRules
                 "InputCore",
 				"Engine",
 				"Slate",
+				"SlateCore",
                 "EditorStyle",
 				"DistCurveEditor",
 				"UnrealEd"
 			}
-			);
+		);
 
 		DynamicallyLoadedModuleNames.AddRange(
 			new string[] {
 				"MainFrame",
 				"PropertyEditor"
 			}
-			);
+		);
 	}
 }
diff --git a/Engine/Source/Editor/ClassViewer/ClassViewer.Build.cs b/Engine/Source/Editor/ClassViewer/ClassViewer.Build.cs
--- a/Engine/Source/Editor/ClassViewer/ClassViewer.Build.cs
+++ b/Engine/Source/Editor/ClassViewer/ClassViewer.Build.cs
@@ -12,7 +12,7 @@ public class ClassViewer : ModuleRules
 				"EditorWidgets",
                 "PropertyEditor"
 			}
-			);
+		);
 
 		PrivateDependencyModuleNames.AddRange(
 			new string[] {
@@ -21,12 +21,13 @@ public class ClassViewer : ModuleRules
 				"Engine",
                 "InputCore",
 				"Slate",
+				"SlateCore",
                 "EditorStyle",
 				"UnrealEd",
 				"PropertyEditor",
 				"GameProjectGeneration"
 			}
-			);
+		);
 
 		DynamicallyLoadedModuleNames.AddRange(
 			new string[] {
@@ -34,6 +35,6 @@ public class ClassViewer : ModuleRules
 				"EditorWidgets",
 				"PropertyEditor"
 			}
-			);
+		);
 	}
 }
diff --git a/Engine/Source/Editor/ContentBrowser/Private/AssetViewWidgets.cpp b/Engine/Source/Editor/ContentBrowser/Private/AssetViewWidgets.cpp
--- a/Engine/Source/Editor/ContentBrowser/Private/AssetViewWidgets.cpp
+++ b/Engine/Source/Editor/ContentBrowser/Private/AssetViewWidgets.cpp
@@ -163,7 +163,7 @@ void SAssetViewItem::Tick( const FGeometry& AllottedGeometry, const double InCur
 	UpdateSourceControlState((float)InDeltaTime);
 }
 
-TSharedPtr<SToolTip> SAssetViewItem::GetToolTip()
+TSharedPtr<IToolTip> SAssetViewItem::GetToolTip()
 {
 	return ShouldAllowToolTip.Get() ? SCompoundWidget::GetToolTip() : NULL;
 }
diff --git a/Engine/Source/Editor/ContentBrowser/Private/AssetViewWidgets.h b/Engine/Source/Editor/ContentBrowser/Private/AssetViewWidgets.h
--- a/Engine/Source/Editor/ContentBrowser/Private/AssetViewWidgets.h
+++ b/Engine/Source/Editor/ContentBrowser/Private/AssetViewWidgets.h
@@ -96,7 +96,7 @@ public:
 	void Construct( const FArguments& InArgs );
 
 	virtual void Tick( const FGeometry& AllottedGeometry, const double InCurrentTime, const float InDeltaTime ) OVERRIDE;
-	virtual TSharedPtr<SToolTip> GetToolTip() OVERRIDE;
+	virtual TSharedPtr<IToolTip> GetToolTip() OVERRIDE;
 	virtual void OnDragEnter( const FGeometry& MyGeometry, const FDragDropEvent& DragDropEvent ) OVERRIDE;
 	virtual void OnDragLeave( const FDragDropEvent& DragDropEvent ) OVERRIDE;
 	virtual FReply OnDragOver( const FGeometry& MyGeometry, const FDragDropEvent& DragDropEvent ) OVERRIDE;
diff --git a/Engine/Source/Editor/DetailCustomizations/DetailCustomizations.Build.cs b/Engine/Source/Editor/DetailCustomizations/DetailCustomizations.Build.cs
--- a/Engine/Source/Editor/DetailCustomizations/DetailCustomizations.Build.cs
+++ b/Engine/Source/Editor/DetailCustomizations/DetailCustomizations.Build.cs
@@ -15,6 +15,7 @@ public class DetailCustomizations : ModuleRules
 				"Engine",
                 "InputCore",
 				"Slate",
+				"SlateCore",
                 "EditorStyle",
 				"UnrealEd",
 				"EditorWidgets",
@@ -35,20 +36,20 @@ public class DetailCustomizations : ModuleRules
                 "InternationalizationSettings",
 				"SourceCodeAccess"
 			}
-			);
+		);
 
 		PrivateIncludePathModuleNames.AddRange(
 			new string[] {
 				"PropertyEditor",
 				"LandscapeEditor",
 			}
-			);
+		);
 
 		DynamicallyLoadedModuleNames.AddRange(
 			new string[] {
 				"PropertyEditor",
 				"Layers"
 			}
-			);
+		);
 	}
 }
diff --git a/Engine/Source/Editor/DistCurveEditor/DistCurveEditor.Build.cs b/Engine/Source/Editor/DistCurveEditor/DistCurveEditor.Build.cs
--- a/Engine/Source/Editor/DistCurveEditor/DistCurveEditor.Build.cs
+++ b/Engine/Source/Editor/DistCurveEditor/DistCurveEditor.Build.cs
@@ -15,17 +15,18 @@ public class DistCurveEditor : ModuleRules
 				"Engine",
                 "InputCore",
 				"Slate",
+				"SlateCore",
                 "EditorStyle",
 				"LevelEditor",
 				"UnrealEd"
 			}
-			);
+		);
 
 		DynamicallyLoadedModuleNames.AddRange(
 			new string[] {
 				"MainFrame",
 				"PropertyEditor"
 			}
-			);
+		);
 	}
 }
diff --git a/Engine/Source/Editor/FoliageEdit/FoliageEdit.build.cs b/Engine/Source/Editor/FoliageEdit/FoliageEdit.build.cs
--- a/Engine/Source/Editor/FoliageEdit/FoliageEdit.build.cs
+++ b/Engine/Source/Editor/FoliageEdit/FoliageEdit.build.cs
@@ -14,10 +14,11 @@ public class FoliageEdit : ModuleRules
 				"Engine",
 				"UnrealEd",
 				"Slate",
+				"SlateCore",
                 "EditorStyle",
 				"RenderCore",
-				"LevelEditor"
+				"LevelEditor",
 			}
-			);
+		);
 	}
 }
diff --git a/Engine/Source/Editor/GameplayTagsEditor/GameplayTagsEditor.Build.cs b/Engine/Source/Editor/GameplayTagsEditor/GameplayTagsEditor.Build.cs
--- a/Engine/Source/Editor/GameplayTagsEditor/GameplayTagsEditor.Build.cs
+++ b/Engine/Source/Editor/GameplayTagsEditor/GameplayTagsEditor.Build.cs
@@ -10,7 +10,7 @@ namespace UnrealBuildTool.Rules
 				new string[] {
 					"Editor/GameplayTagsEditor/Private",
 				}
-				);
+			);
 
 
 			PrivateDependencyModuleNames.AddRange(
@@ -24,6 +24,7 @@ namespace UnrealBuildTool.Rules
 					"GameplayTags",
                     "InputCore",
 					"Slate",
+					"SlateCore",
                     "EditorStyle",
 					"BlueprintGraph",
 					"KismetCompiler",
@@ -31,8 +32,7 @@ namespace UnrealBuildTool.Rules
 					"MainFrame",
 					"UnrealEd",
 				}
-				);
-
+			);
 		}
 	}
 }
diff --git a/Engine/Source/Editor/GeometryMode/GeometryMode.Build.cs b/Engine/Source/Editor/GeometryMode/GeometryMode.Build.cs
--- a/Engine/Source/Editor/GeometryMode/GeometryMode.Build.cs
+++ b/Engine/Source/Editor/GeometryMode/GeometryMode.Build.cs
@@ -13,13 +13,15 @@ public class GeometryMode : ModuleRules
 				"Engine",
                 "InputCore",
 				"Slate",
+				"SlateCore",
                 "EditorStyle",
 				"UnrealEd",
 				"RenderCore",
 				"LevelEditor"
 			}
-			);
-        DynamicallyLoadedModuleNames.AddRange(
+		);
+
+		DynamicallyLoadedModuleNames.AddRange(
             new string[] {
 				"PropertyEditor",
                 "BspMode"
diff --git a/Engine/Source/Editor/GraphEditor/GraphEditor.Build.cs b/Engine/Source/Editor/GraphEditor/GraphEditor.Build.cs
--- a/Engine/Source/Editor/GraphEditor/GraphEditor.Build.cs
+++ b/Engine/Source/Editor/GraphEditor/GraphEditor.Build.cs
@@ -33,6 +33,7 @@ public class GraphEditor : ModuleRules
 				"Engine",
                 "InputCore",
 				"Slate",
+				"SlateCore",
                 "EditorStyle",
 				"UnrealEd",
 				"AssetRegistry",
@@ -43,7 +44,7 @@ public class GraphEditor : ModuleRules
 				"RenderCore",
 				"RHI"
 			}
-			);
+		);
 
 		DynamicallyLoadedModuleNames.AddRange(
 			new string[] {
@@ -51,6 +52,6 @@ public class GraphEditor : ModuleRules
 				"Documentation",
 				"EditorWidgets"
 			}
-			);
+		);
 	}
 }
diff --git a/Engine/Source/Editor/GraphEditor/Private/SGraphNode.cpp b/Engine/Source/Editor/GraphEditor/Private/SGraphNode.cpp
--- a/Engine/Source/Editor/GraphEditor/Private/SGraphNode.cpp
+++ b/Engine/Source/Editor/GraphEditor/Private/SGraphNode.cpp
@@ -334,9 +334,9 @@ FReply SGraphNode::OnMouseButtonDoubleClick( const FGeometry& InMyGeometry, cons
 	return FReply::Unhandled();
 }
 
-TSharedPtr<SToolTip> SGraphNode::GetToolTip()
+TSharedPtr<IToolTip> SGraphNode::GetToolTip()
 {
-	TSharedPtr<SToolTip> CurrentTooltip = SWidget::GetToolTip();
+	TSharedPtr<IToolTip> CurrentTooltip = SWidget::GetToolTip();
 	if (!CurrentTooltip.IsValid())
 	{
 		TSharedPtr<SToolTip> ComplexTooltip = GetComplexTooltip();
diff --git a/Engine/Source/Editor/GraphEditor/Public/SGraphNode.h b/Engine/Source/Editor/GraphEditor/Public/SGraphNode.h
--- a/Engine/Source/Editor/GraphEditor/Public/SGraphNode.h
+++ b/Engine/Source/Editor/GraphEditor/Public/SGraphNode.h
@@ -74,7 +74,7 @@ public:
 	virtual FReply OnMouseButtonUp( const FGeometry& MyGeometry, const FPointerEvent& MouseEvent ) OVERRIDE;
 	virtual FReply OnMouseButtonDoubleClick( const FGeometry& InMyGeometry, const FPointerEvent& InMouseEvent ) OVERRIDE;
 
-	virtual TSharedPtr<SToolTip> GetToolTip() OVERRIDE;
+	virtual TSharedPtr<IToolTip> GetToolTip() OVERRIDE;
 	virtual void OnToolTipClosing() OVERRIDE;
 
 	virtual void Tick( const FGeometry& AllottedGeometry, const double InCurrentTime, const float InDeltaTime ) OVERRIDE;
diff --git a/Engine/Source/Editor/InputBindingEditor/InputBindingEditor.Build.cs b/Engine/Source/Editor/InputBindingEditor/InputBindingEditor.Build.cs
--- a/Engine/Source/Editor/InputBindingEditor/InputBindingEditor.Build.cs
+++ b/Engine/Source/Editor/InputBindingEditor/InputBindingEditor.Build.cs
@@ -9,8 +9,9 @@ public class InputBindingEditor : ModuleRules
 		PublicDependencyModuleNames.AddRange(
 			new string[]
 			{
+				"EditorStyle",
 				"Slate",
-                "EditorStyle"
+                "SlateCore",
 			}
 		);
 
diff --git a/Engine/Source/Editor/Kismet/Private/SBlueprintPalette.cpp b/Engine/Source/Editor/Kismet/Private/SBlueprintPalette.cpp
--- a/Engine/Source/Editor/Kismet/Private/SBlueprintPalette.cpp
+++ b/Engine/Source/Editor/Kismet/Private/SBlueprintPalette.cpp
@@ -360,9 +360,9 @@ static void GetPaletteItemIcon(TSharedPtr<FEdGraphSchemaAction> ActionIn, UBluep
  * @param  OldToolTip	The tooltip that you're replacing (we fold it into the new one)/
  * @return The newly created tooltip (now with the action's path tacked on the bottom).
  */
-static TSharedRef<SToolTip> ConstructToolTipWithActionPath(TSharedPtr<FEdGraphSchemaAction> ActionIn, TSharedPtr<SToolTip> OldToolTip)
+static TSharedRef<IToolTip> ConstructToolTipWithActionPath(TSharedPtr<FEdGraphSchemaAction> ActionIn, TSharedPtr<IToolTip> OldToolTip)
 {
-	TSharedRef<SToolTip> NewToolTip = OldToolTip.ToSharedRef();
+	TSharedRef<IToolTip> NewToolTip = OldToolTip.ToSharedRef();
 
 	FFavoritedBlueprintPaletteItem ActionItem(ActionIn);
 	if (ActionItem.IsValid())
@@ -376,7 +376,7 @@ static TSharedRef<SToolTip> ConstructToolTipWithActionPath(TSharedPtr<FEdGraphSc
 			SNew(SVerticalBox)
 			+SVerticalBox::Slot()
 			[
-				OldToolTip->GetContent()
+				OldToolTip->GetContentWidget()
 			]
 
 			+SVerticalBox::Slot()
@@ -968,7 +968,7 @@ TSharedRef<SWidget> SBlueprintPaletteItem::CreateTextSlotWidget(const FSlateFont
 	if (GEditor->EditorUserSettings->bDisplayActionListItemRefIds && ActionPtr.IsValid())
 	{
 		check(InlineRenameWidget.IsValid());
-		TSharedPtr<SToolTip> ExistingToolTip = InlineRenameWidget->GetToolTip();
+		TSharedPtr<IToolTip> ExistingToolTip = InlineRenameWidget->GetToolTip();
 
 		DisplayWidget->AddSlot(0)
 			[
diff --git a/Engine/Source/Editor/KismetWidgets/KismetWidgets.Build.cs b/Engine/Source/Editor/KismetWidgets/KismetWidgets.Build.cs
--- a/Engine/Source/Editor/KismetWidgets/KismetWidgets.Build.cs
+++ b/Engine/Source/Editor/KismetWidgets/KismetWidgets.Build.cs
@@ -15,24 +15,25 @@ public class KismetWidgets : ModuleRules
                 "InputCore",
 				"Engine",
 				"Slate",
+				"SlateCore",
                 "EditorStyle",
 				"UnrealEd",
 				"BlueprintGraph",
 			}
-			);
+		);
 
         PrivateIncludePathModuleNames.AddRange(
             new string[] {
 				"ContentBrowser",
                 "AssetTools",
 			}
-            );
+		);
 
         DynamicallyLoadedModuleNames.AddRange(
             new string[] {
 			    "ContentBrowser",
                 "AssetTools",
 			}
-            );
+		);
 	}
 }
diff --git a/Engine/Source/Editor/Layers/Layers.Build.cs b/Engine/Source/Editor/Layers/Layers.Build.cs
--- a/Engine/Source/Editor/Layers/Layers.Build.cs
+++ b/Engine/Source/Editor/Layers/Layers.Build.cs
@@ -13,17 +13,18 @@ public class Layers : ModuleRules
 				"Engine",
                 "InputCore",
 				"Slate",
+				"SlateCore",
                 "EditorStyle",
 				"UnrealEd",
-				"SceneOutliner"
+				"SceneOutliner",
 			}
-			);
+		);
 
 		DynamicallyLoadedModuleNames.AddRange(
 			new string[] {
 				"PropertyEditor",
-				"SceneOutliner"
+				"SceneOutliner",
 			}
-			);
+		);
 	}
 }
diff --git a/Engine/Source/Editor/LevelEditor/LevelEditor.Build.cs b/Engine/Source/Editor/LevelEditor/LevelEditor.Build.cs
--- a/Engine/Source/Editor/LevelEditor/LevelEditor.Build.cs
+++ b/Engine/Source/Editor/LevelEditor/LevelEditor.Build.cs
@@ -30,6 +30,8 @@ public class LevelEditor : ModuleRules
 				"DesktopPlatform",
                 "InputCore",
 				"Slate",
+				"SlateCore",
+				"SlateReflector",
                 "EditorStyle",
 				"Engine",
 				"MessageLog",
diff --git a/Engine/Source/Editor/LevelEditor/Private/LevelEditor.cpp b/Engine/Source/Editor/LevelEditor/Private/LevelEditor.cpp
--- a/Engine/Source/Editor/LevelEditor/Private/LevelEditor.cpp
+++ b/Engine/Source/Editor/LevelEditor/Private/LevelEditor.cpp
@@ -15,6 +15,7 @@
 #include "LevelViewportActions.h"
 #include "GlobalEditorCommonCommands.h"
 #include "IUserFeedbackModule.h"
+#include "SlateReflector.h"
 
 // @todo Editor: remove this circular dependency
 #include "Editor/MainFrame/Public/Interfaces/IMainFrameModule.h"
@@ -90,7 +91,6 @@ TSharedRef<SDockTab> SpawnLevelEditor( const FSpawnTabArgs& InArgs )
 				.AutoWidth()
 				[
 					SNew(SBox)
-					.Padding( FMargin( 0, 1, 20, 0 ) )
 					.Visibility( EVisibility::HitTestInvisible )
 					[
 						SNew( STextBlock )
@@ -102,6 +102,8 @@ TSharedRef<SDockTab> SpawnLevelEditor( const FSpawnTabArgs& InArgs )
 
 				+SHorizontalBox::Slot()
 				.AutoWidth()
+				.Padding(16.0f, 0.0f, 0.0f, 0.0f)
+				.VAlign(VAlign_Center)
 				[
 					UserFeedbackWidget
 				]
@@ -142,7 +144,7 @@ void FLevelEditorModule::StartupModule()
 	FGlobalTabmanager::Get()->RegisterTabSpawner("LevelEditor", FOnSpawnTab::CreateStatic( &SpawnLevelEditor ) )
 		.SetDisplayName( NSLOCTEXT("LevelEditor", "LevelEditorTab", "Level Editor") );
 
-	FSlateApplication::Get().RegisterWidgetReflectorTabSpawner( MenuStructure.GetDeveloperToolsCategory() );
+	FModuleManager::LoadModuleChecked<ISlateReflectorModule>("SlateReflector").RegisterTabSpawner(MenuStructure.GetDeveloperToolsCategory());
 
 	FMessageLogModule& MessageLogModule = FModuleManager::LoadModuleChecked<FMessageLogModule>("MessageLog");
 	MessageLogModule.RegisterLogListing("BuildAndSubmitErrors", LOCTEXT("BuildAndSubmitErrors", "Build and Submit Errors"));
@@ -187,7 +189,7 @@ void FLevelEditorModule::ShutdownModule()
 	if ( FSlateApplication::IsInitialized() )
 	{
 		FGlobalTabmanager::Get()->UnregisterTabSpawner("LevelEditor");
-		FSlateApplication::Get().UnregisterWidgetReflectorTabSpawner();
+		FModuleManager::LoadModuleChecked<ISlateReflectorModule>("SlateReflector").UnregisterTabSpawner();
 	}	
 
 	FLevelEditorCommands::Unregister();
diff --git a/Engine/Source/Editor/NewLevelDialog/NewLevelDialog.Build.cs b/Engine/Source/Editor/NewLevelDialog/NewLevelDialog.Build.cs
--- a/Engine/Source/Editor/NewLevelDialog/NewLevelDialog.Build.cs
+++ b/Engine/Source/Editor/NewLevelDialog/NewLevelDialog.Build.cs
@@ -8,15 +8,16 @@ public class NewLevelDialog : ModuleRules
 	{
 		PrivateDependencyModuleNames.AddRange(
 			new string[] {
-				"Core", 
-				"CoreUObject", 
+				"Core",
+				"CoreUObject",
                 "InputCore",
-				"RenderCore", 
-				"Engine", 
-				"Slate", 
+				"RenderCore",
+				"Engine",
+				"Slate",
+				"SlateCore",
                 "EditorStyle",
-				"UnrealEd"
+				"UnrealEd",
 			}
-			);
+		);
 	}
 }
diff --git a/Engine/Source/Editor/NiagaraEditor/NiagaraEditor.Build.cs b/Engine/Source/Editor/NiagaraEditor/NiagaraEditor.Build.cs
--- a/Engine/Source/Editor/NiagaraEditor/NiagaraEditor.Build.cs
+++ b/Engine/Source/Editor/NiagaraEditor/NiagaraEditor.Build.cs
@@ -14,16 +14,16 @@ public class NiagaraEditor : ModuleRules
                 "InputCore",
 				"RenderCore",
 				"Slate", 
+				"SlateCore",
                 "EditorStyle",
 				"UnrealEd", 
 				"GraphEditor", 
 				"VectorVM"
 			}
-			);
+		);
 
 		PublicIncludePathModuleNames.Add("LevelEditor");
 
 		DynamicallyLoadedModuleNames.Add("WorkspaceMenuStructure");
-
 	}
 }
diff --git a/Engine/Source/Editor/PListEditor/PListEditor.Build.cs b/Engine/Source/Editor/PListEditor/PListEditor.Build.cs
--- a/Engine/Source/Editor/PListEditor/PListEditor.Build.cs
+++ b/Engine/Source/Editor/PListEditor/PListEditor.Build.cs
@@ -7,30 +7,28 @@ public class PListEditor : ModuleRules
 	public PListEditor( TargetInfo Target )
 	{
 		PrivateDependencyModuleNames.AddRange(
-			new string[] 
-			{ 
-				"Core", 
-				"CoreUObject", 
-				"Engine", 
+			new string[] { 
+				"Core",
+				"CoreUObject",
+				"Engine",
                 "InputCore",
-				"Slate", 
+				"Slate",
+				"SlateCore",
                 "EditorStyle",
-				"UnrealEd", 
+				"UnrealEd",
 				"DesktopPlatform",
-				"XmlParser"
+				"XmlParser",
 			}
 		);
 
 		PrivateIncludePathModuleNames.AddRange(
-			new string[]
-			{
+			new string[] {
 				"MainFrame",
 			}
 		);
 
 		DynamicallyLoadedModuleNames.AddRange(
-			new string[]
-			{
+			new string[] {
 				"MainFrame",
 			}
 		);
diff --git a/Engine/Source/Editor/PhAT/PhAT.Build.cs b/Engine/Source/Editor/PhAT/PhAT.Build.cs
--- a/Engine/Source/Editor/PhAT/PhAT.Build.cs
+++ b/Engine/Source/Editor/PhAT/PhAT.Build.cs
@@ -16,19 +16,20 @@ public class PhAT : ModuleRules
                 "InputCore",
 				"RenderCore",
 				"Slate",
+				"SlateCore",
                 "EditorStyle",
 				"LevelEditor",
 				"UnrealEd",
-                "Kismet"
+                "Kismet",
 			}
-			);
+		);
 
 		DynamicallyLoadedModuleNames.AddRange(
 			new string[] {
 				"MainFrame",
 				"PropertyEditor",
-                "MeshUtilities"
+                "MeshUtilities",
 			}
-			);
+		);
 	}
 }
diff --git a/Engine/Source/Editor/ReferenceViewer/ReferenceViewer.Build.cs b/Engine/Source/Editor/ReferenceViewer/ReferenceViewer.Build.cs
--- a/Engine/Source/Editor/ReferenceViewer/ReferenceViewer.Build.cs
+++ b/Engine/Source/Editor/ReferenceViewer/ReferenceViewer.Build.cs
@@ -13,10 +13,11 @@ namespace UnrealBuildTool.Rules
 				    "Engine",
                     "InputCore",
 				    "Slate",
+					"SlateCore",
                     "EditorStyle",
 				    "RenderCore",
                     "GraphEditor",
-				    "UnrealEd"
+				    "UnrealEd",
 			    }
 			);
 
diff --git a/Engine/Source/Editor/SkillSystemEditor/SkillSystemEditor.Build.cs b/Engine/Source/Editor/SkillSystemEditor/SkillSystemEditor.Build.cs
--- a/Engine/Source/Editor/SkillSystemEditor/SkillSystemEditor.Build.cs
+++ b/Engine/Source/Editor/SkillSystemEditor/SkillSystemEditor.Build.cs
@@ -10,8 +10,7 @@ namespace UnrealBuildTool.Rules
                 new string[] {
 					"Editor/SkillSystemTagsEditor/Private",
 				}
-                );
-
+			);
 
             PrivateDependencyModuleNames.AddRange(
                 new string[]
@@ -24,6 +23,7 @@ namespace UnrealBuildTool.Rules
 					"SkillSystem",
                     "InputCore",
 					"Slate",
+					"SlateCore",
                     "EditorStyle",
 					"BlueprintGraph",
 					"KismetCompiler",
@@ -31,7 +31,7 @@ namespace UnrealBuildTool.Rules
 					"MainFrame",
 					"UnrealEd",
 				}
-                );
+			);
 		}
 	}
 }
\ No newline at end of file
diff --git a/Engine/Source/Editor/SourceControlWindows/SourceControlWindows.Build.cs b/Engine/Source/Editor/SourceControlWindows/SourceControlWindows.Build.cs
--- a/Engine/Source/Editor/SourceControlWindows/SourceControlWindows.Build.cs
+++ b/Engine/Source/Editor/SourceControlWindows/SourceControlWindows.Build.cs
@@ -12,12 +12,13 @@ public class SourceControlWindows : ModuleRules
 				"CoreUObject", 
                 "InputCore",
 				"Engine", 
-				"Slate", 
+				"Slate",
+				"SlateCore",
                 "EditorStyle",
 				"SourceControl", 
 				"AssetTools",
                 "UnrealEd"			// We need this dependency here because we use PackageTools.
 			}
-			);
+		);
 	}
 }
diff --git a/Engine/Source/Editor/TaskBrowser/TaskBrowser.Build.cs b/Engine/Source/Editor/TaskBrowser/TaskBrowser.Build.cs
--- a/Engine/Source/Editor/TaskBrowser/TaskBrowser.Build.cs
+++ b/Engine/Source/Editor/TaskBrowser/TaskBrowser.Build.cs
@@ -15,10 +15,10 @@ public class TaskBrowser : ModuleRules
 				"Engine",
                 "InputCore",
 				"Slate",
+				"SlateCore",
                 "EditorStyle",
 				"UnrealEd"
 			}
-			);
-
+		);
 	}
 }
diff --git a/Engine/Source/Editor/UMGEditor/Private/SUMGDesigner.cpp b/Engine/Source/Editor/UMGEditor/Private/SUMGDesigner.cpp
--- a/Engine/Source/Editor/UMGEditor/Private/SUMGDesigner.cpp
+++ b/Engine/Source/Editor/UMGEditor/Private/SUMGDesigner.cpp
@@ -3,10 +3,8 @@
 #include "UMGEditorPrivatePCH.h"
 
 #include "SUMGDesigner.h"
-
 #include "BlueprintEditor.h"
 #include "SKismetInspector.h"
-#include "PrimitiveTypes.h"
 
 #define LOCTEXT_NAMESPACE "UMG"
 
diff --git a/Engine/Source/Editor/UMGEditor/UMGEditor.Build.cs b/Engine/Source/Editor/UMGEditor/UMGEditor.Build.cs
--- a/Engine/Source/Editor/UMGEditor/UMGEditor.Build.cs
+++ b/Engine/Source/Editor/UMGEditor/UMGEditor.Build.cs
@@ -48,7 +48,9 @@ public class UMGEditor : ModuleRules
                 "DistCurveEditor",
 				"ContentBrowser",
 				"WorkspaceMenuStructure",
-				"EditorStyle"
+				"EditorStyle",
+				"Slate",
+				"SlateCore",
 			}
 			);
 
diff --git a/Engine/Source/Editor/UndoHistory/UndoHistory.Build.cs b/Engine/Source/Editor/UndoHistory/UndoHistory.Build.cs
--- a/Engine/Source/Editor/UndoHistory/UndoHistory.Build.cs
+++ b/Engine/Source/Editor/UndoHistory/UndoHistory.Build.cs
@@ -7,27 +7,25 @@ namespace UnrealBuildTool.Rules
 		public UndoHistory(TargetInfo Target)
 		{
 			PublicDependencyModuleNames.AddRange(
-				new string[]
-				{
+				new string[] {
 					"Core",
 				}
 			);
 
 			PrivateDependencyModuleNames.AddRange(
-				new string[]
-				{
+				new string[] {
 					"CoreUObject",
 					"EditorStyle",
 					"Engine",
 					"InputCore",
 					"Slate",
+					"SlateCore",
 					"UnrealEd",
 				}
 			);
 
 			PrivateIncludePaths.AddRange(
-				new string[]
-				{
+				new string[] {
 					"Editor/UndoHistory/Private",
 					"Editor/UndoHistory/Private/Widgets",
 				}
diff --git a/Engine/Source/Editor/UserFeedback/Private/UserFeedback.cpp b/Engine/Source/Editor/UserFeedback/Private/UserFeedback.cpp
--- a/Engine/Source/Editor/UserFeedback/Private/UserFeedback.cpp
+++ b/Engine/Source/Editor/UserFeedback/Private/UserFeedback.cpp
@@ -627,15 +627,14 @@ public:
 	{
 		if (FEngineAnalytics::IsAvailable() )
 		{
-			return 
-				SNew(SHorizontalBox)
-				+SHorizontalBox::Slot()
-				.Padding(FMargin(0, 0, 10, 2))
+			return SNew(SHorizontalBox)
+
+			+ SHorizontalBox::Slot()
+				.Padding(0.0f, 0.0f, 10.0f, 0.0f)
 				.VAlign(VAlign_Center)
 				[
 					SNew(SUserFeedbackButtonWidget).Context(Context)
 				];
-
 		}	
 		else
 		{
diff --git a/Engine/Source/Editor/UserFeedback/UserFeedback.Build.cs b/Engine/Source/Editor/UserFeedback/UserFeedback.Build.cs
--- a/Engine/Source/Editor/UserFeedback/UserFeedback.Build.cs
+++ b/Engine/Source/Editor/UserFeedback/UserFeedback.Build.cs
@@ -8,9 +8,11 @@ public class UserFeedback : ModuleRules
 	{
 		PrivateDependencyModuleNames.AddRange(
 			new string[] {
+				"SlateCore",
 				"UnrealEd",
 			}
 		);
+
 		PublicDependencyModuleNames.AddRange(
 			new string[] {
 				"Core",
diff --git a/Engine/Source/Editor/WorkspaceMenuStructure/WorkspaceMenuStructure.Build.cs b/Engine/Source/Editor/WorkspaceMenuStructure/WorkspaceMenuStructure.Build.cs
--- a/Engine/Source/Editor/WorkspaceMenuStructure/WorkspaceMenuStructure.Build.cs
+++ b/Engine/Source/Editor/WorkspaceMenuStructure/WorkspaceMenuStructure.Build.cs
@@ -11,9 +11,9 @@ public class WorkspaceMenuStructure : ModuleRules
 		PrivateDependencyModuleNames.AddRange(
 			new string[] {
 				"Core",
+				"EditorStyle",
 				"Slate",
-                "EditorStyle"
 			}
-			);
+		);
 	}
 }
diff --git a/Engine/Source/Editor/WorldBrowser/Private/SWorldGridItem.cpp b/Engine/Source/Editor/WorldBrowser/Private/SWorldGridItem.cpp
--- a/Engine/Source/Editor/WorldBrowser/Private/SWorldGridItem.cpp
+++ b/Engine/Source/Editor/WorldBrowser/Private/SWorldGridItem.cpp
@@ -297,7 +297,7 @@ FSlateRect SWorldGridItem::GetItemRect() const
 	return FSlateRect(LevelPos, LevelPos + LevelSize);
 }
 
-TSharedPtr<SToolTip> SWorldGridItem::GetToolTip()
+TSharedPtr<IToolTip> SWorldGridItem::GetToolTip()
 {
 	// Hide tooltip in case item is being dragged now
 	if (LevelModel->GetLevelTranslationDelta().Size() > KINDA_SMALL_NUMBER)
diff --git a/Engine/Source/Editor/WorldBrowser/WorldBrowser.Build.cs b/Engine/Source/Editor/WorldBrowser/WorldBrowser.Build.cs
--- a/Engine/Source/Editor/WorldBrowser/WorldBrowser.Build.cs
+++ b/Engine/Source/Editor/WorldBrowser/WorldBrowser.Build.cs
@@ -11,40 +11,41 @@ public class WorldBrowser : ModuleRules
         PrivateIncludePathModuleNames.AddRange(
         new string[] {
 				"AssetTools",
-                "MeshUtilities",	
-                "ContentBrowser"	
+                "MeshUtilities",
+                "ContentBrowser",
 			}
         );
      
         PrivateDependencyModuleNames.AddRange(
             new string[] {    
                 "Core", 
-                "CoreUObject", 
-                "RenderCore", 
+                "CoreUObject",
+                "RenderCore",
                 "InputCore",
-                "Engine", 
-                "Slate", 
+                "Engine",
+                "Slate",
+				"SlateCore",
                 "EditorStyle",
                 "UnrealEd",
-                "GraphEditor", 
-                "LevelEditor", 
-                "PropertyEditor", 
+                "GraphEditor",
+                "LevelEditor",
+                "PropertyEditor",
                 "DesktopPlatform",
                 "MainFrame",
                 "SourceControl",
 				"SourceControlWindows",
                 "RawMesh",
-                "MeshUtilities"
+                "MeshUtilities",
             }
-            );
+		);
 
         DynamicallyLoadedModuleNames.AddRange(
             new string[] {
 				"AssetTools",
 				"SceneOutliner",
                 "MeshUtilities",
-                "ContentBrowser"
+                "ContentBrowser",
 			}
-            );
+		);
     }
-}
\ No newline at end of file
+}
diff --git a/Engine/Source/Programs/CrashReporter/CrashReportClient/CrashReportClient.Build.cs b/Engine/Source/Programs/CrashReporter/CrashReportClient/CrashReportClient.Build.cs
--- a/Engine/Source/Programs/CrashReporter/CrashReportClient/CrashReportClient.Build.cs
+++ b/Engine/Source/Programs/CrashReporter/CrashReportClient/CrashReportClient.Build.cs
@@ -25,6 +25,8 @@ public class CrashReportClient : ModuleRules
 				new string[] 
 				{
 					"Slate",
+					"SlateCore",
+					"SlateReflector",
 					"StandaloneRenderer",
 					"DesktopPlatform",
 					"MessageLog",
diff --git a/Engine/Source/Programs/CrashReporter/CrashReportClient/Private/CrashReportClientApp.cpp b/Engine/Source/Programs/CrashReporter/CrashReportClient/Private/CrashReportClientApp.cpp
--- a/Engine/Source/Programs/CrashReporter/CrashReportClient/Private/CrashReportClientApp.cpp
+++ b/Engine/Source/Programs/CrashReporter/CrashReportClient/Private/CrashReportClientApp.cpp
@@ -11,6 +11,7 @@
 #include "AsyncWork.h"
 
 #include "MainLoopTiming.h"
+#include "SlateReflector.h"
 
 /** Default main window size */
 const FVector2D InitialWindowDimensions(640, 560);
@@ -122,7 +123,7 @@ void RunCrashReportClient(const TCHAR* CommandLine)
 				SNew(SWindow)
 				.ClientSize(FVector2D(800, 600))
 				[
-					FSlateApplication::Get().GetWidgetReflector()
+					FModuleManager::LoadModuleChecked<ISlateReflectorModule>().GetWidgetReflector()
 				]);
 		}
 
diff --git a/Engine/Source/Programs/SlateViewer/Private/SlateViewerApp.cpp b/Engine/Source/Programs/SlateViewer/Private/SlateViewerApp.cpp
--- a/Engine/Source/Programs/SlateViewer/Private/SlateViewerApp.cpp
+++ b/Engine/Source/Programs/SlateViewer/Private/SlateViewerApp.cpp
@@ -2,7 +2,6 @@
 
 #include "SlateViewerApp.h"
 #include "RequiredProgramMainCPPInclude.h"
-#include "Runtime/Slate/Private/SWidgetReflector.h"
 
 
 IMPLEMENT_APPLICATION(SlateViewer, "SlateViewer");
@@ -24,6 +23,7 @@ void RunSlateViewer( const TCHAR* CommandLine )
 
 	// set the application name
 	FGlobalTabmanager::Get()->SetApplicationTitle(NSLOCTEXT("SlateViewer", "AppTitle", "Slate Viewer"));
+	FModuleManager::LoadModuleChecked<ISlateReflectorModule>("SlateReflector").RegisterTabSpawner(WorkspaceMenu::DeveloperMenu);
 
 	// Bring up the test suite.
 	RestoreSlateTestSuite();
diff --git a/Engine/Source/Programs/SlateViewer/Public/SlateViewerApp.h b/Engine/Source/Programs/SlateViewer/Public/SlateViewerApp.h
--- a/Engine/Source/Programs/SlateViewer/Public/SlateViewerApp.h
+++ b/Engine/Source/Programs/SlateViewer/Public/SlateViewerApp.h
@@ -4,6 +4,7 @@
 #pragma once
 
 #include "SlateViewer.h"
+#include "SlateReflector.h"
 
 /**
  * Run the SlateViewer 
diff --git a/Engine/Source/Programs/SlateViewer/SlateViewer.Build.cs b/Engine/Source/Programs/SlateViewer/SlateViewer.Build.cs
--- a/Engine/Source/Programs/SlateViewer/SlateViewer.Build.cs
+++ b/Engine/Source/Programs/SlateViewer/SlateViewer.Build.cs
@@ -8,9 +8,16 @@ public class SlateViewer : ModuleRules
 	{
 		PublicIncludePaths.Add("Runtime/Launch/Public");
 
-		PrivateIncludePathModuleNames.AddRange(new string[] { "Slate" });
-
-		PrivateDependencyModuleNames.AddRange(new string[] { "Core", "Slate", "StandaloneRenderer", "Projects" });
+		PrivateDependencyModuleNames.AddRange(
+			new string[] {
+				"Core",
+				"Projects",
+				"Slate",
+				"SlateCore",
+				"SlateReflector",
+				"StandaloneRenderer",
+			}
+		);
 
 		PrivateIncludePaths.Add("Runtime/Launch/Private");		// For LaunchEngineLoop.cpp include
 
diff --git a/Engine/Source/Programs/SymbolDebugger/SymbolDebugger.Build.cs b/Engine/Source/Programs/SymbolDebugger/SymbolDebugger.Build.cs
--- a/Engine/Source/Programs/SymbolDebugger/SymbolDebugger.Build.cs
+++ b/Engine/Source/Programs/SymbolDebugger/SymbolDebugger.Build.cs
@@ -9,10 +9,10 @@ public class SymbolDebugger : ModuleRules
 		PublicIncludePaths.Add("Runtime/Launch/Public");
 
 		PrivateDependencyModuleNames.AddRange(
-			new string[] 
-			{
+			new string[] {
 				"Core",
 				"Slate",
+				"SlateCore",
 				"EditorStyle",
 				"StandaloneRenderer",
 				"DesktopPlatform",
@@ -22,7 +22,7 @@ public class SymbolDebugger : ModuleRules
 				"MessageLog",
 				"Projects",
 			}
-			);
+		);
 
 		PrivateIncludePaths.Add("Runtime/Launch/Private");		// For LaunchEngineLoop.cpp include
 
diff --git a/Engine/Source/Programs/UnrealFrontend/Private/UnrealFrontendMain.cpp b/Engine/Source/Programs/UnrealFrontend/Private/UnrealFrontendMain.cpp
--- a/Engine/Source/Programs/UnrealFrontend/Private/UnrealFrontendMain.cpp
+++ b/Engine/Source/Programs/UnrealFrontend/Private/UnrealFrontendMain.cpp
@@ -14,6 +14,7 @@
 #include "StatsData.h"
 #include "StatsFile.h"
 #include "EditorStyle.h"
+#include "SlateReflector.h"
 
 
 IMPLEMENT_APPLICATION(UnrealFrontend, "UnrealFrontend");
@@ -369,7 +370,7 @@ void RunUI()
 	FModuleManager::Get().LoadModule("SettingsEditor");
 
 	// Create developer tools menu with widget reflector.
-	FSlateApplication::Get().RegisterWidgetReflectorTabSpawner(WorkspaceMenu::DeveloperTools);
+	FModuleManager::LoadModuleChecked<ISlateReflectorModule>("SlateReflector").RegisterTabSpawner(WorkspaceMenu::DeveloperTools);
 
 	// check to see if want the widget reflector
 	const bool bAllowDebugTools = FParse::Param(FCommandLine::Get(), TEXT("DebugTools"));
diff --git a/Engine/Source/Runtime/Engine/Engine.Build.cs b/Engine/Source/Runtime/Engine/Engine.Build.cs
--- a/Engine/Source/Runtime/Engine/Engine.Build.cs
+++ b/Engine/Source/Runtime/Engine/Engine.Build.cs
@@ -58,7 +58,8 @@ public class Engine : ModuleRules
 			new string[] {
 				"Networking",
 				"Sockets",
-				"Slate",
+				"SlateCore",
+				"SlateReflector",
 				"VectorVM",
 			}
 		);
diff --git a/Engine/Source/Runtime/Engine/Private/AVIWriter.cpp b/Engine/Source/Runtime/Engine/Private/AVIWriter.cpp
--- a/Engine/Source/Runtime/Engine/Private/AVIWriter.cpp
+++ b/Engine/Source/Runtime/Engine/Private/AVIWriter.cpp
@@ -23,8 +23,8 @@ typedef TCHAR* PTCHAR;
 #include "HideWindowsPlatformTypes.h"
 
 #include "CaptureSource.h"
+#include "SlateCore.h"
 
-#include "Renderer.h"
 
 #define g_wszCapture    L"Capture Filter"
 
diff --git a/Engine/Source/Runtime/Engine/Private/CheatManager.cpp b/Engine/Source/Runtime/Engine/Private/CheatManager.cpp
--- a/Engine/Source/Runtime/Engine/Private/CheatManager.cpp
+++ b/Engine/Source/Runtime/Engine/Private/CheatManager.cpp
@@ -71,6 +71,7 @@
 #include "Components/WindDirectionalSourceComponent.h"
 #include "Components/TimelineComponent.h"
 #include "Slate.h"
+#include "SlateReflector.h"
 #include "AI/NavDataGenerator.h"
 #include "OnlineSubsystemUtils.h"
 
@@ -886,11 +887,11 @@ void UCheatManager::WidgetReflector()
 	if( !WidgetReflectorWindow.IsValid() )
 	{
 		const TSharedRef<SWindow> ReflectorWindow = SNew(SWindow)
-		.AutoCenter(EAutoCenter::PrimaryWorkArea)
-		.ClientSize(FVector2D(600,400))
-		[
-			FSlateApplication::Get().GetWidgetReflector()
-		];
+			.AutoCenter(EAutoCenter::PrimaryWorkArea)
+			.ClientSize(FVector2D(600,400))
+			[
+				FModuleManager::LoadModuleChecked<ISlateReflectorModule>("SlateReflector").GetWidgetReflector()
+			];
 		
 		WidgetReflectorWindow = ReflectorWindow;
 		
@@ -1206,4 +1207,4 @@ void UCheatManager::LogOutBugItGoToLogFile( const FString& InScreenShotDesc, con
 #endif // ALLOW_DEBUG_FILES
 }
 
-#undef LOCTEXT_NAMESPACE
\ No newline at end of file
+#undef LOCTEXT_NAMESPACE
diff --git a/Engine/Source/Runtime/HeadMountedDisplay/HeadMountedDisplay.Build.cs b/Engine/Source/Runtime/HeadMountedDisplay/HeadMountedDisplay.Build.cs
--- a/Engine/Source/Runtime/HeadMountedDisplay/HeadMountedDisplay.Build.cs
+++ b/Engine/Source/Runtime/HeadMountedDisplay/HeadMountedDisplay.Build.cs
@@ -10,8 +10,17 @@ public class HeadMountedDisplay : ModuleRules
 			new string[] {
 				"Runtime/HeadMountedDisplay/Public"
 			}
-			);
-		PrivateDependencyModuleNames.AddRange( new string[] { "Core", "CoreUObject", "Engine", "Slate" } );
+		);
+
+		PrivateDependencyModuleNames.AddRange(
+			new string[] {
+				"Core",
+				"CoreUObject",
+				"Engine",
+				"Slate",
+				"SlateCore",
+			}
+		);
 
         PCHUsage = PCHUsageMode.NoSharedPCHs;
 	}
diff --git a/Engine/Source/Runtime/Slate/Private/DockingFramework/SDockingTabStack.cpp b/Engine/Source/Runtime/Slate/Private/DockingFramework/SDockingTabStack.cpp
--- a/Engine/Source/Runtime/Slate/Private/DockingFramework/SDockingTabStack.cpp
+++ b/Engine/Source/Runtime/Slate/Private/DockingFramework/SDockingTabStack.cpp
@@ -4,9 +4,34 @@
 #include "DockingPrivate.h"
 
 
+#define LOCTEXT_NAMESPACE "DockTabStack"
+
+
 static const FVector2D ContextButtonTargetSize(24,24);
 static const float TriggerAreaFraction = 0.24f;
 
+
+/**
+ * Like a missing widget, but says it's a document area
+ */
+class SDocumentAreaWidget
+{
+public:
+	static TSharedRef<class SWidget> MakeDocumentAreaWidget()
+	{
+		return 
+			SNew(SBorder)
+			.HAlign(HAlign_Center)
+			.VAlign(VAlign_Center)
+			[
+				SNew( STextBlock )
+				.Text(LOCTEXT("DocumentArea", "Document Area") )
+				.TextStyle( FCoreStyle::Get(), "EmbossedText" )
+			];
+	}
+};
+
+
 void SDockingTabStack::Construct( const FArguments& InArgs, const TSharedRef<FTabManager::FStack>& PersistentNode  )
 {
 	Tabs = PersistentNode->Tabs;
@@ -16,9 +41,7 @@ void SDockingTabStack::Construct( const FArguments& InArgs, const TSharedRef<FTa
 
 	InlineContentAreaLeft = NULL;
 	InlineContentAreaRight = NULL;
-
-	TitleAreaLeftSlot = NULL;
-	TitleAreaRightSlot = NULL;
+	TitleBarSlot = NULL;
 
 	this->TabStackGeometry = FGeometry();
 
@@ -52,7 +75,46 @@ void SDockingTabStack::Construct( const FArguments& InArgs, const TSharedRef<FTa
 
 	const FButtonStyle* const UnhideTabWellButtonStyle = &FCoreStyle::Get().GetWidgetStyle< FButtonStyle >( "Docking.UnhideTabwellButton" );
 
-	this->ChildSlot
+	// create inline title bar content
+	TitleBarContent = SNew(SHorizontalBox)
+		.Visibility(EVisibility::SelfHitTestInvisible)
+
+	+ SHorizontalBox::Slot()
+		.AutoWidth()
+		.Expose(InlineContentAreaLeft)
+
+	+ SHorizontalBox::Slot() 
+		.FillWidth(1.0f)
+		.VAlign(VAlign_Bottom)
+		.Padding(4.0f, 0.0f, 0.0f, 0.0f)
+		[
+			SNew(SVerticalBox)
+				.Visibility(EVisibility::SelfHitTestInvisible)
+
+			+ SVerticalBox::Slot()
+				.AutoHeight()
+				[
+					SNew(SSpacer)
+						.Visibility(this, &SDockingTabStack::GetMaximizeSpacerVisibility)
+						.Size(FVector2D(0.0f, 10.0f))
+				]
+
+			+ SVerticalBox::Slot()
+				.AutoHeight()
+				[
+					// TabWell
+					SAssignNew(TabWell, SDockingTabWell)
+						.ParentStackNode(SharedThis(this))
+				]
+		]
+
+	+ SHorizontalBox::Slot()
+		.AutoWidth()
+		.Expose(InlineContentAreaRight)
+		.Padding(5.0f, 0.0f, 0.0f, 0.0f)
+		.VAlign((VAlign_Center));
+
+	ChildSlot
 	[
 		SNew(SVerticalBox)
 		.Visibility( EVisibility::SelfHitTestInvisible )
@@ -72,64 +134,9 @@ void SDockingTabStack::Construct( const FArguments& InArgs, const TSharedRef<FTa
 				.Visibility( EVisibility::SelfHitTestInvisible )
 				+SVerticalBox::Slot()
 				.AutoHeight()
+				.Expose(TitleBarSlot)
 				[
-
-					SNew(SHorizontalBox)
-					.Visibility( EVisibility::SelfHitTestInvisible )
-
-					// Left title area (if we have one)
-					+SHorizontalBox::Slot()
-					.AutoWidth()
-					.Expose( TitleAreaLeftSlot )
-					.HAlign(HAlign_Left)
-					.VAlign(VAlign_Top)
-					[
-						SNew(SSpacer)	// Will be filled in later, in SetParentNode()
-					]
-
-					// Inline Content Left
-					+SHorizontalBox::Slot()
-					.AutoWidth()
-					.Expose( InlineContentAreaLeft )
-					.Padding( 0,0,5,0 )				
-
-					+SHorizontalBox::Slot() 
-					.FillWidth(1)
-					.VAlign(VAlign_Bottom)
-					[
-						SNew(SVerticalBox)
-						.Visibility( EVisibility::SelfHitTestInvisible )
-						+SVerticalBox::Slot()
-						. AutoHeight()
-						[
-							SNew(SSpacer)
-							.Visibility(this, &SDockingTabStack::GetMaximizeSpacerVisibility)
-							.Size(FVector2D(0.0f, 10.0f))
-						]
-						+SVerticalBox::Slot()
-						. AutoHeight()
-						[
-							// TabWell
-							SAssignNew(TabWell, SDockingTabWell)
-							.ParentStackNode( SharedThis(this) )
-						]
-					]
-
-					// Inline Content Right
-					+SHorizontalBox::Slot()
-					.AutoWidth()
-					.Expose( InlineContentAreaRight )
-					.Padding(5,0,0,0)
-
-					// Right title area (if we have one)
-					+SHorizontalBox::Slot()
-					.AutoWidth()
-					.Expose( TitleAreaRightSlot )
-					.HAlign(HAlign_Right)
-					.VAlign(VAlign_Top)
-					[
-						SNew(SSpacer)	// Will be filled in later, in SetParentNode()
-					]
+					SNew(SSpacer)
 				]
 				+SVerticalBox::Slot()
 				.AutoHeight()
@@ -151,7 +158,7 @@ void SDockingTabStack::Construct( const FArguments& InArgs, const TSharedRef<FTa
 				[
 					// Selected CONTENT AREA
 					SNew(STextBlock)
-					. Text( NSLOCTEXT("DockTabStack", "EmptyTabMessage", "Empty Tab!") )
+					. Text( LOCTEXT("EmptyTabMessage", "Empty Tab!") )
 				]
 			]
 			+SOverlay::Slot()
@@ -489,11 +496,11 @@ TSharedRef<SWidget> SDockingTabStack::MakeContextMenu()
 	const bool bCloseSelfOnly = false;
 	FMenuBuilder MenuBuilder( bCloseAfterSelection, NULL, TSharedPtr<FExtender>(), bCloseSelfOnly, &FCoreStyle::Get() );
 	{
-		MenuBuilder.BeginSection("DockingTabStackOptions", NSLOCTEXT("Docking", "TabOptionsHeading", "Options") );
+		MenuBuilder.BeginSection("DockingTabStackOptions", LOCTEXT("TabOptionsHeading", "Options") );
 		{
 			MenuBuilder.AddMenuEntry(
-				NSLOCTEXT("Docking", "HideTabWell", "Hide Tab"),
-				NSLOCTEXT("Docking", "HideTabWellTooltip", "Hide the tabs to save room."),
+				LOCTEXT("HideTabWell", "Hide Tab"),
+				LOCTEXT("HideTabWellTooltip", "Hide the tabs to save room."),
 				FSlateIcon(),
 				FUIAction(
 					FExecuteAction::CreateSP( this, &SDockingTabStack::ToggleTabWellVisibility ),
@@ -506,8 +513,8 @@ TSharedRef<SWidget> SDockingTabStack::MakeContextMenu()
 		MenuBuilder.BeginSection("DockingTabStackCloseTabs");
 		{
 			MenuBuilder.AddMenuEntry(
-				NSLOCTEXT("Docking", "CloseTab", "Close"),
-				NSLOCTEXT("Docking", "CloseTabTooltil", "Close this tab."),
+				LOCTEXT("CloseTab", "Close"),
+				LOCTEXT("CloseTabTooltil", "Close this tab."),
 				FSlateIcon(),
 				FUIAction(
 					FExecuteAction::CreateSP( this, &SDockingTabStack::CloseForegroundTab ),
@@ -522,8 +529,8 @@ TSharedRef<SWidget> SDockingTabStack::MakeContextMenu()
 				const ETabsToClose TabsToClose = CloseDocumentTabs;
 
 				MenuBuilder.AddMenuEntry(
-					NSLOCTEXT("Docking", "CloseOtherTabs", "Close Other Tabs"),
-					NSLOCTEXT("Docking", "CloseOtherTabsTooltil", "Closes all tabs except for the active tab."),
+					LOCTEXT("CloseOtherTabs", "Close Other Tabs"),
+					LOCTEXT("CloseOtherTabsTooltil", "Closes all tabs except for the active tab."),
 					FSlateIcon(),
 					FUIAction(
 						FExecuteAction::CreateSP( this, &SDockingTabStack::CloseAllButForegroundTab, TabsToClose ),
@@ -705,18 +712,11 @@ void SDockingTabStack::SetParentNode( TSharedRef<class SDockingSplitter> InParen
 		//                  and only the first one!  Currently, all SDockingAreas with a parent window set will get
 		//                  title area widgets added!
 		const TSharedRef<SWindow>& ParentWindow = DockArea->GetParentWindow().ToSharedRef();
+		TSharedPtr<IWindowTitleBar> TitleBar;
+		TSharedRef<SWidget> TitleBarWidget = FSlateApplication::Get().MakeWindowTitleBar(ParentWindow, TitleBarContent, TitleBar);
 
-		TSharedPtr< SWidget > LeftContent, RightContent;
-		TSharedPtr< SWidget > CenterContent;	// NOTE: We don't care about center title area content.  No space to put it!
-		ParentWindow->MakeTitleBarContentWidgets( LeftContent, CenterContent, RightContent );
-		
-		TitleAreaLeftSlot->Widget = LeftContent.IsValid() ? LeftContent.ToSharedRef() : TSharedRef<SWidget>( SNew( SSpacer ) );
-		TitleAreaRightSlot->Widget = RightContent.IsValid() ? RightContent.ToSharedRef() : TSharedRef<SWidget>( SNew( SSpacer ) );
-	}
-	else
-	{
-		TitleAreaLeftSlot->Widget = SNew( SSpacer );
-		TitleAreaRightSlot->Widget = SNew( SSpacer );
+		(*TitleBarSlot)[TitleBarWidget];
+		ParentWindow->SetTitleBar(TitleBar);
 	}
 }
 
@@ -958,3 +958,5 @@ FString SDockingTabStack::ShowPersistentTabs() const
 
 #endif
 
+
+#undef LOCTEXT_NAMESPACE
diff --git a/Engine/Source/Runtime/Slate/Private/DockingFramework/SDockingTabStack.h b/Engine/Source/Runtime/Slate/Private/DockingFramework/SDockingTabStack.h
--- a/Engine/Source/Runtime/Slate/Private/DockingFramework/SDockingTabStack.h
+++ b/Engine/Source/Runtime/Slate/Private/DockingFramework/SDockingTabStack.h
@@ -165,12 +165,8 @@ private:
 	/** The borders that hold any potential inline content areas. */
 	SHorizontalBox::FSlot* InlineContentAreaLeft;
 	SHorizontalBox::FSlot* InlineContentAreaRight;
-
-	/** Slot for left title area content */
-	SHorizontalBox::FSlot* TitleAreaLeftSlot;
-
-	/** Slot for right title area content */
-	SHorizontalBox::FSlot* TitleAreaRightSlot;
+	SVerticalBox::FSlot* TitleBarSlot;
+	TSharedPtr<SWidget> TitleBarContent;
 
 	TSharedPtr<SBorder> ContentSlot;
 
