WidgetReflector: fixed list columns being cut off / having poor initial sizes

diff --git a/Engine/Source/Developer/SlateReflector/Private/Widgets/SWidgetReflector.cpp b/Engine/Source/Developer/SlateReflector/Private/Widgets/SWidgetReflector.cpp
--- a/Engine/Source/Developer/SlateReflector/Private/Widgets/SWidgetReflector.cpp
+++ b/Engine/Source/Developer/SlateReflector/Private/Widgets/SWidgetReflector.cpp
@@ -99,11 +99,17 @@ void SWidgetReflector::Construct( const FArguments& InArgs )
 
 								+ SHeaderRow::Column("WidgetName")
 									.DefaultLabel(LOCTEXT("WidgetName", "Widget Name").ToString())
-									.FillWidth(0.50f)
+									.FillWidth(0.65f)
 
 								+ SHeaderRow::Column("ForegroundColor")
-									.DefaultLabel(LOCTEXT("ForegroundColor", "Foreground Color").ToString())
-									.FixedWidth(20.0f)
+									.FixedWidth(24.0f)
+									.VAlignHeader(VAlign_Center)
+									.HeaderContent()
+									[
+										SNew(STextBlock)
+											.Text(LOCTEXT("ForegroundColor", "FG").ToString())
+											.ToolTipText(LOCTEXT("ForegroundColorToolTip", "Foreground Color").ToString())
+									]
 
 								+ SHeaderRow::Column("Visibility")
 									.DefaultLabel(LOCTEXT("Visibility", "Visibility" ).ToString())
@@ -111,7 +117,7 @@ void SWidgetReflector::Construct( const FArguments& InArgs )
 
 								+ SHeaderRow::Column("WidgetInfo")
 									.DefaultLabel(LOCTEXT("WidgetInfo", "Widget Info" ).ToString())
-									.FillWidth(30.0f)
+									.FillWidth(0.35f)
 							)
 					]
 
