Better spacing on DPI scaler control in the Widget Reflector.

diff --git a/Engine/Source/Developer/SlateReflector/Private/Widgets/SWidgetReflector.cpp b/Engine/Source/Developer/SlateReflector/Private/Widgets/SWidgetReflector.cpp
--- a/Engine/Source/Developer/SlateReflector/Private/Widgets/SWidgetReflector.cpp
+++ b/Engine/Source/Developer/SlateReflector/Private/Widgets/SWidgetReflector.cpp
@@ -34,15 +34,15 @@ void SWidgetReflector::Construct( const FArguments& InArgs )
 							]
 
 						+ SHorizontalBox::Slot()
-							.AutoWidth()
-							[
-								SNew(SSpinBox<float>)
-									.Value(this, &SWidgetReflector::HandleAppScaleSliderValue)
-									.MinValue(0.1f)
-									.MaxValue(3.0f)
-									.Delta(0.01f)
-									.OnValueChanged(this, &SWidgetReflector::HandleAppScaleSliderChanged)
-							]
+						.MaxWidth(250)
+						[
+							SNew(SSpinBox<float>)
+								.Value(this, &SWidgetReflector::HandleAppScaleSliderValue)
+								.MinValue(0.1f)
+								.MaxValue(3.0f)
+								.Delta(0.01f)
+								.OnValueChanged(this, &SWidgetReflector::HandleAppScaleSliderChanged)
+						]
 					]
 
 				+ SVerticalBox::Slot()
