Slate WidgetReflector now shows widget's memory address.

diff --git a/Engine/Source/Developer/SlateReflector/Private/Widgets/SWidgetReflector.cpp b/Engine/Source/Developer/SlateReflector/Private/Widgets/SWidgetReflector.cpp
--- a/Engine/Source/Developer/SlateReflector/Private/Widgets/SWidgetReflector.cpp
+++ b/Engine/Source/Developer/SlateReflector/Private/Widgets/SWidgetReflector.cpp
@@ -113,7 +113,11 @@ void SWidgetReflector::Construct( const FArguments& InArgs )
 
 								+ SHeaderRow::Column("WidgetInfo")
 									.DefaultLabel(LOCTEXT("WidgetInfo", "Widget Info" ).ToString())
-									.FillWidth(0.35f)
+									.FillWidth(0.25f)
+
+								+ SHeaderRow::Column("Address")
+									.DefaultLabel( LOCTEXT("Address", "Address") )
+									.FillWidth( 0.10f )
 							)
 					]
 
diff --git a/Engine/Source/Developer/SlateReflector/Private/Widgets/SWidgetReflectorTreeWidgetItem.h b/Engine/Source/Developer/SlateReflector/Private/Widgets/SWidgetReflectorTreeWidgetItem.h
--- a/Engine/Source/Developer/SlateReflector/Private/Widgets/SWidgetReflectorTreeWidgetItem.h
+++ b/Engine/Source/Developer/SlateReflector/Private/Widgets/SWidgetReflectorTreeWidgetItem.h
@@ -104,6 +104,19 @@ public:
 						.Size(FVector2D(16.0f, 16.0f))
 				];
 		}
+		else if (ColumnName == "Address")
+		{
+			const TSharedPtr<SWidget> TheWidget = WidgetInfo.Get()->Widget.Pin();
+			const FText Address = (TheWidget.IsValid())
+				? FText::FromString( FString::Printf(TEXT("0x%08X"), TheWidget.Get()) ) 
+				: NSLOCTEXT("SWidgetReflector","nullptr","nullptr");
+			
+			return
+				SNew(SHyperlink)
+				.ToolTipText(NSLOCTEXT("SWidgetReflector", "ClickToCopy", "Click to copy address."))
+				.Text( Address )
+				.OnNavigate_Lambda( [Address](){ FPlatformMisc::ClipboardCopy(*Address.ToString()); } ) ;
+		}
 		else
 		{
 			return SNullWidget::NullWidget;
