SlateReflector: Made headers compilable individually; code & documentation clenaup

diff --git a/Engine/Source/Developer/SlateReflector/Private/SlateReflectorModule.cpp b/Engine/Source/Developer/SlateReflector/Private/SlateReflectorModule.cpp
--- a/Engine/Source/Developer/SlateReflector/Private/SlateReflectorModule.cpp
+++ b/Engine/Source/Developer/SlateReflector/Private/SlateReflectorModule.cpp
@@ -1,6 +1,7 @@
 // Copyright 1998-2014 Epic Games, Inc. All Rights Reserved.
 
 #include "SlateReflectorPrivatePCH.h"
+#include "ISlateReflectorModule.h"
 #include "SDockTab.h"
 #include "ModuleManager.h"
 
@@ -18,7 +19,7 @@ public:
 
 	// ISlateReflectorModule interface
 
-	virtual TSharedRef<SWidget> GetWidgetReflector( ) override
+	virtual TSharedRef<SWidget> GetWidgetReflector() override
 	{
 		TSharedPtr<SWidgetReflector> WidgetReflector = WidgetReflectorPtr.Pin();
 
@@ -40,7 +41,7 @@ public:
 			.SetIcon(FSlateIcon(FCoreStyle::Get().GetStyleSetName(), "WidgetReflector.TabIcon"));
 	}
 
-	virtual void UnregisterTabSpawner( ) 
+	virtual void UnregisterTabSpawner() 
 	{
 		FGlobalTabmanager::Get()->UnregisterNomadTabSpawner("WidgetReflector");
 	}
@@ -49,9 +50,9 @@ public:
 
 	// IModuleInterface interface
 
-	virtual void StartupModule( ) override { }
+	virtual void StartupModule() override { }
 
-	virtual void ShutdownModule( ) override
+	virtual void ShutdownModule() override
 	{
 		UnregisterTabSpawner();
 	}
@@ -69,7 +70,7 @@ private:
 
 private:
 
-	// Holds the widget reflector singleton.
+	/** Holds the widget reflector singleton. */
 	TWeakPtr<SWidgetReflector> WidgetReflectorPtr;
 };
 
diff --git a/Engine/Source/Developer/SlateReflector/Private/SlateReflectorPrivatePCH.h b/Engine/Source/Developer/SlateReflector/Private/SlateReflectorPrivatePCH.h
--- a/Engine/Source/Developer/SlateReflector/Private/SlateReflectorPrivatePCH.h
+++ b/Engine/Source/Developer/SlateReflector/Private/SlateReflectorPrivatePCH.h
@@ -2,12 +2,12 @@
 
 #pragma once
 
-#include "SlateReflector.h"
-
 
 /* Dependencies
  *****************************************************************************/
 
+#include "Core.h"
+#include "CoreUObject.h"
 #include "InputCore.h"
 #include "SlateBasics.h"
 #include "SlateCore.h"
diff --git a/Engine/Source/Developer/SlateReflector/Private/Widgets/SWidgetReflector.cpp b/Engine/Source/Developer/SlateReflector/Private/Widgets/SWidgetReflector.cpp
--- a/Engine/Source/Developer/SlateReflector/Private/Widgets/SWidgetReflector.cpp
+++ b/Engine/Source/Developer/SlateReflector/Private/Widgets/SWidgetReflector.cpp
@@ -4,24 +4,26 @@
 
 
 #define LOCTEXT_NAMESPACE "SWidgetReflector"
+#define WITH_EVENT_LOGGING 0
 
 static const int32 MaxLoggedEvents = 100;
 
-#define WITH_EVENT_LOGGING 0
+
+/* Local helpers
+ *****************************************************************************/
 
 struct FLoggedEvent
 {
 	FLoggedEvent( const FInputEvent& InEvent, const FReplyBase& InReply )
-	: Event( InEvent )
-	, Handler( InReply.GetHandler() )
-	, EventText( InEvent.ToText() )
-	, HandlerText( InReply.GetHandler().IsValid() ? FText::FromString(InReply.GetHandler()->ToString()) : LOCTEXT("NullHandler", "null") )
-	{		
-	}
+		: Event(InEvent)
+		, Handler(InReply.GetHandler())
+		, EventText(InEvent.ToText())
+		, HandlerText(InReply.GetHandler().IsValid() ? FText::FromString(InReply.GetHandler()->ToString()) : LOCTEXT("NullHandler", "null"))
+	{ }
 
-	FText ToText( )
+	FText ToText()
 	{
-		return FText::Format( NSLOCTEXT("","","{0}  |  {1}"), EventText, HandlerText );
+		return FText::Format(NSLOCTEXT("","","{0}  |  {1}"), EventText, HandlerText);
 	}
 	
 	FInputEvent Event;
@@ -44,128 +46,128 @@ void SWidgetReflector::Construct( const FArguments& InArgs )
 	this->ChildSlot
 	[
 		SNew(SBorder)
-		.BorderImage(FCoreStyle::Get().GetBrush("ToolPanel.GroupBorder"))
-		[
-			SNew(SVerticalBox)
-				
-			+ SVerticalBox::Slot()
-			.AutoHeight()
+			.BorderImage(FCoreStyle::Get().GetBrush("ToolPanel.GroupBorder"))
 			[
-				SNew(SHorizontalBox)
+				SNew(SVerticalBox)
+				
+				+ SVerticalBox::Slot()
+					.AutoHeight()
+					[
+						SNew(SHorizontalBox)
 					
-				+ SHorizontalBox::Slot()
-				.AutoWidth()
-				[
-					SNew(STextBlock)
-						.Text(LOCTEXT("AppScale", "Application Scale: ").ToString())
-				]
-
-				+ SHorizontalBox::Slot()
-				.MaxWidth(250)
-				[
-					SNew(SSpinBox<float>)
-					.Value(this, &SWidgetReflector::HandleAppScaleSliderValue)
-					.MinValue(0.1f)
-					.MaxValue(3.0f)
-					.Delta(0.01f)
-					.OnValueChanged(this, &SWidgetReflector::HandleAppScaleSliderChanged)
-				]
-			]
-
-			+ SVerticalBox::Slot()
-			.AutoHeight()
-			[
-				SNew(SHorizontalBox)
-
-				+ SHorizontalBox::Slot()
-				.AutoWidth()
-				.Padding(5.0f)
-				[
-					// Check box that controls LIVE MODE
-					SNew(SCheckBox)
-						.IsChecked(this, &SWidgetReflector::HandleFocusCheckBoxIsChecked)
-						.OnCheckStateChanged(this, &SWidgetReflector::HandleFocusCheckBoxCheckedStateChanged)
-						[
-							SNew(STextBlock)
-								.Text(LOCTEXT("ShowFocus", "Show Focus").ToString())
-						]
-				]
-
-				+ SHorizontalBox::Slot()
-				.AutoWidth()
-				.Padding(5.0f)
-				[
-					// Check box that controls PICKING A WIDGET TO INSPECT
-					SNew(SButton)
-						.OnClicked(this, &SWidgetReflector::HandlePickButtonClicked)
-						.ButtonColorAndOpacity(this, &SWidgetReflector::HandlePickButtonColorAndOpacity)
-						[
-							SNew(STextBlock)
-								.Text(this, &SWidgetReflector::HandlePickButtonText)
-						]
-				]
-			]
+						+ SHorizontalBox::Slot()
+							.AutoWidth()
+							[
+								SNew(STextBlock)
+									.Text(LOCTEXT("AppScale", "Application Scale: ").ToString())
+							]
+
+						+ SHorizontalBox::Slot()
+							.MaxWidth(250)
+							[
+								SNew(SSpinBox<float>)
+									.Value(this, &SWidgetReflector::HandleAppScaleSliderValue)
+									.MinValue(0.1f)
+									.MaxValue(3.0f)
+									.Delta(0.01f)
+									.OnValueChanged(this, &SWidgetReflector::HandleAppScaleSliderChanged)
+							]
+					]
 
-			+ SVerticalBox::Slot()
-			.FillHeight(1.0f)
-			[
-				// The tree view that shows all the info that we capture.
-				SAssignNew(ReflectorTree, SReflectorTree)
-				.ItemHeight(24.0f)
-				.TreeItemsSource(&ReflectorTreeRoot)
-				.OnGenerateRow(this, &SWidgetReflector::HandleReflectorTreeGenerateRow)
-				.OnGetChildren(this, &SWidgetReflector::HandleReflectorTreeGetChildren)
-				.OnSelectionChanged(this, &SWidgetReflector::HandleReflectorTreeSelectionChanged)
-				.HeaderRow
-				(
-					SNew(SHeaderRow)
-
-					+ SHeaderRow::Column("WidgetName")
-					.DefaultLabel(LOCTEXT("WidgetName", "Widget Name").ToString())
-					.FillWidth(0.65f)
-
-					+ SHeaderRow::Column("ForegroundColor")
-					.FixedWidth(24.0f)
-					.VAlignHeader(VAlign_Center)
-					.HeaderContent()
+				+ SVerticalBox::Slot()
+					.AutoHeight()
 					[
-						SNew(STextBlock)
-						.Text(LOCTEXT("ForegroundColor", "FG").ToString())
-						.ToolTipText(LOCTEXT("ForegroundColorToolTip", "Foreground Color").ToString())
+						SNew(SHorizontalBox)
+
+						+ SHorizontalBox::Slot()
+							.AutoWidth()
+							.Padding(5.0f)
+							[
+								// Check box that controls LIVE MODE
+								SNew(SCheckBox)
+									.IsChecked(this, &SWidgetReflector::HandleFocusCheckBoxIsChecked)
+									.OnCheckStateChanged(this, &SWidgetReflector::HandleFocusCheckBoxCheckedStateChanged)
+									[
+										SNew(STextBlock)
+											.Text(LOCTEXT("ShowFocus", "Show Focus").ToString())
+									]
+							]
+
+						+ SHorizontalBox::Slot()
+							.AutoWidth()
+							.Padding(5.0f)
+							[
+								// Check box that controls PICKING A WIDGET TO INSPECT
+								SNew(SButton)
+									.OnClicked(this, &SWidgetReflector::HandlePickButtonClicked)
+									.ButtonColorAndOpacity(this, &SWidgetReflector::HandlePickButtonColorAndOpacity)
+									[
+										SNew(STextBlock)
+											.Text(this, &SWidgetReflector::HandlePickButtonText)
+									]
+							]
 					]
 
-					+ SHeaderRow::Column("Visibility")
-					.DefaultLabel(LOCTEXT("Visibility", "Visibility" ).ToString())
-					.FixedWidth(125.0f)
-
-					+ SHeaderRow::Column("WidgetInfo")
-					.DefaultLabel(LOCTEXT("WidgetInfo", "Widget Info" ).ToString())
-					.FillWidth(0.25f)
-
-					+ SHeaderRow::Column("Address")
-					.DefaultLabel( LOCTEXT("Address", "Address") )
-					.FillWidth( 0.10f )
-				)
-			]
+				+ SVerticalBox::Slot()
+					.FillHeight(1.0f)
+					[
+						// The tree view that shows all the info that we capture.
+						SAssignNew(ReflectorTree, SReflectorTree)
+							.ItemHeight(24.0f)
+							.TreeItemsSource(&ReflectorTreeRoot)
+							.OnGenerateRow(this, &SWidgetReflector::HandleReflectorTreeGenerateRow)
+							.OnGetChildren(this, &SWidgetReflector::HandleReflectorTreeGetChildren)
+							.OnSelectionChanged(this, &SWidgetReflector::HandleReflectorTreeSelectionChanged)
+							.HeaderRow
+							(
+								SNew(SHeaderRow)
+
+								+ SHeaderRow::Column("WidgetName")
+									.DefaultLabel(LOCTEXT("WidgetName", "Widget Name").ToString())
+									.FillWidth(0.65f)
+
+								+ SHeaderRow::Column("ForegroundColor")
+									.FixedWidth(24.0f)
+									.VAlignHeader(VAlign_Center)
+									.HeaderContent()
+									[
+										SNew(STextBlock)
+											.Text(LOCTEXT("ForegroundColor", "FG").ToString())
+											.ToolTipText(LOCTEXT("ForegroundColorToolTip", "Foreground Color").ToString())
+									]
+
+								+ SHeaderRow::Column("Visibility")
+									.DefaultLabel(LOCTEXT("Visibility", "Visibility" ).ToString())
+									.FixedWidth(125.0f)
+
+								+ SHeaderRow::Column("WidgetInfo")
+									.DefaultLabel(LOCTEXT("WidgetInfo", "Widget Info" ).ToString())
+									.FillWidth(0.25f)
+
+								+ SHeaderRow::Column("Address")
+									.DefaultLabel( LOCTEXT("Address", "Address") )
+									.FillWidth( 0.10f )
+							)
+					]
 
-			#if WITH_EVENT_LOGGING
-			+ SVerticalBox::Slot()
-			.FillHeight(1.0f)
-			[
-				SAssignNew(EventListView, SListView<TSharedRef<FLoggedEvent>>)
-				.ListItemsSource( &LoggedEvents )
-				.OnGenerateRow(this, &SWidgetReflector::GenerateEventLogRow)
-			]
-			#endif //WITH_EVENT_LOGGING
+				#if WITH_EVENT_LOGGING
+				+ SVerticalBox::Slot()
+					.FillHeight(1.0f)
+					[
+						SAssignNew(EventListView, SListView<TSharedRef<FLoggedEvent>>)
+							.ListItemsSource( &LoggedEvents )
+							.OnGenerateRow(this, &SWidgetReflector::GenerateEventLogRow)
+					]
+				#endif //WITH_EVENT_LOGGING
 			
-			+ SVerticalBox::Slot()
-			.AutoHeight()
-			[
-				// Frame rate counter
-				SNew(STextBlock)
-					.Text(this, &SWidgetReflector::HandleFrameRateText)
+				+ SVerticalBox::Slot()
+					.AutoHeight()
+					[
+						// Frame rate counter
+						SNew(STextBlock)
+							.Text(this, &SWidgetReflector::HandleFrameRateText)
+					]
 			]
-		]
 	];
 }
 
@@ -178,6 +180,7 @@ void SWidgetReflector::Tick( const FGeometry& AllottedGeometry, const double InC
 	SCompoundWidget::Tick(AllottedGeometry, InCurrentTime, InDeltaTime);
 }
 
+
 void SWidgetReflector::OnEventProcessed( const FInputEvent& Event, const FReplyBase& InReply )
 {
 	#if WITH_EVENT_LOGGING
@@ -186,9 +189,9 @@ void SWidgetReflector::OnEventProcessed( const FInputEvent& Event, const FReplyB
 			LoggedEvents.Empty();
 		}
 
-		LoggedEvents.Add( MakeShareable( new FLoggedEvent( Event, InReply ) ) );
+		LoggedEvents.Add(MakeShareable(new FLoggedEvent(Event, InReply)));
 		EventListView->RequestListRefresh();
-		EventListView->RequestScrollIntoView( LoggedEvents.Last() );
+		EventListView->RequestScrollIntoView(LoggedEvents.Last());
 	#endif //WITH_EVENT_LOGGING
 }
 
@@ -246,14 +249,14 @@ int32 SWidgetReflector::Visualize( const FWidgetPath& InWidgetsToVisualize, FSla
 TSharedRef<SToolTip> SWidgetReflector::GenerateToolTipForReflectorNode( TSharedPtr<FReflectorNode> InReflectorNode )
 {
 	return SNew(SToolTip)
-	[
-		SNew(SReflectorToolTipWidget)
-			.WidgetInfoToVisualize(InReflectorNode)
-	];
+		[
+			SNew(SReflectorToolTipWidget)
+				.WidgetInfoToVisualize(InReflectorNode)
+		];
 }
 
 
-void SWidgetReflector::VisualizeAsTree( const TArray< TSharedPtr<FReflectorNode> >& WidgetPathToVisualize )
+void SWidgetReflector::VisualizeAsTree( const TArray<TSharedPtr<FReflectorNode>>& WidgetPathToVisualize )
 {
 	const FLinearColor TopmostWidgetColor(1.0f, 0.0f, 0.0f);
 	const FLinearColor LeafmostWidgetColor(0.0f, 1.0f, 0.0f);
@@ -301,7 +304,7 @@ int32 SWidgetReflector::VisualizePickAsRectangles( const FWidgetPath& InWidgetsT
 			FCoreStyle::Get().GetBrush(TEXT("Debug.Border")),
 			InWidgetsToVisualize.TopLevelWindow->GetClippingRectangleInWindow(),
 			ESlateDrawEffect::None,
-			FMath::Lerp( TopmostWidgetColor, LeafmostWidgetColor, ColorFactor )
+			FMath::Lerp(TopmostWidgetColor, LeafmostWidgetColor, ColorFactor)
 		);
 	}
 
@@ -353,7 +356,7 @@ void SWidgetReflector::HandleFocusCheckBoxCheckedStateChanged( ESlateCheckBoxSta
 }
 
 
-FString SWidgetReflector::HandleFrameRateText( ) const
+FString SWidgetReflector::HandleFrameRateText() const
 {
 	FString MyString;
 #if 0 // the new stats system does not support this
@@ -364,7 +367,7 @@ FString SWidgetReflector::HandleFrameRateText( ) const
 }
 
 
-FString SWidgetReflector::HandlePickButtonText( ) const
+FString SWidgetReflector::HandlePickButtonText() const
 {
 	static const FString NotPicking = LOCTEXT("PickWidget", "Pick Widget").ToString();
 	static const FString Picking = LOCTEXT("PickingWidget", "Picking (Esc to Stop)").ToString();
@@ -381,12 +384,13 @@ TSharedRef<ITableRow> SWidgetReflector::HandleReflectorTreeGenerateRow( TSharedP
 		.SourceCodeAccessor(SourceAccessDelegate);
 }
 
+
 TSharedRef<ITableRow> SWidgetReflector::GenerateEventLogRow( TSharedRef<FLoggedEvent> InLoggedEvent, const TSharedRef<STableViewBase>& OwnerTable )
 {
 	return SNew(STableRow<TSharedRef<FLoggedEvent>>, OwnerTable)
 	[
 		SNew(STextBlock)
-		.Text( InLoggedEvent->ToText() )
+			.Text(InLoggedEvent->ToText())
 	];
 }
 
diff --git a/Engine/Source/Developer/SlateReflector/Private/Widgets/SWidgetReflector.h b/Engine/Source/Developer/SlateReflector/Private/Widgets/SWidgetReflector.h
--- a/Engine/Source/Developer/SlateReflector/Private/Widgets/SWidgetReflector.h
+++ b/Engine/Source/Developer/SlateReflector/Private/Widgets/SWidgetReflector.h
@@ -4,8 +4,11 @@
 
 #include "IWidgetReflector.h"
 
+
+// forward declarations
 struct FLoggedEvent;
 
+
 /**
  * Abstract base class for widget reflectors.
  */
@@ -40,22 +43,22 @@ public:
 
 	virtual void OnEventProcessed( const FInputEvent& Event, const FReplyBase& InReply ) override;
 
-	virtual bool IsInPickingMode( ) const override
+	virtual bool IsInPickingMode() const override
 	{
 		return bIsPicking;
 	}
 
-	virtual bool IsShowingFocus( ) const override
+	virtual bool IsShowingFocus() const override
 	{
 		return bShowFocus;
 	}
 
-	virtual bool IsVisualizingLayoutUnderCursor( ) const override
+	virtual bool IsVisualizingLayoutUnderCursor() const override
 	{
 		return bIsPicking;
 	}
 
-	virtual void OnWidgetPicked( ) override
+	virtual void OnWidgetPicked() override
 	{
 		bIsPicking = false;
 	}
@@ -110,32 +113,32 @@ protected:
 
 private:
 
-	// Callback for changing the application scale slider
+	/** Callback for changing the application scale slider. */
 	void HandleAppScaleSliderChanged( float NewValue )
 	{
 		FSlateApplication::Get().SetApplicationScale(NewValue);
 	}
 
-	// Callback for getting the value of the application scale slider.
-	float HandleAppScaleSliderValue( ) const
+	/** Callback for getting the value of the application scale slider. */
+	float HandleAppScaleSliderValue() const
 	{
 		return FSlateApplication::Get().GetApplicationScale();
 	}
 
-	// Callback for checked state changes of the focus check box.
+	/** Callback for checked state changes of the focus check box. */
 	void HandleFocusCheckBoxCheckedStateChanged( ESlateCheckBoxState::Type NewValue );
 
-	// Callback for getting the checked state of the focus check box.
-	ESlateCheckBoxState::Type HandleFocusCheckBoxIsChecked( ) const
+	/** Callback for getting the checked state of the focus check box. */
+	ESlateCheckBoxState::Type HandleFocusCheckBoxIsChecked() const
 	{
 		return bShowFocus ? ESlateCheckBoxState::Checked : ESlateCheckBoxState::Unchecked;
 	}
 
-	// Callback for getting the text of the frame rate text block.
-	FString HandleFrameRateText( ) const;
+	/** Callback for getting the text of the frame rate text block. */
+	FString HandleFrameRateText() const;
 
-	// Callback for clicking the pick button.
-	FReply HandlePickButtonClicked( )
+	/** Callback for clicking the pick button. */
+	FReply HandlePickButtonClicked()
 	{
 		bIsPicking = !bIsPicking;
 
@@ -147,27 +150,27 @@ private:
 		return FReply::Handled();
 	}
 
-	// Callback for getting the color of the pick button text.
-	FSlateColor HandlePickButtonColorAndOpacity( ) const
+	/** Callback for getting the color of the pick button text. */
+	FSlateColor HandlePickButtonColorAndOpacity() const
 	{
 		return bIsPicking
 			? FCoreStyle::Get().GetSlateColor("SelectionColor")
 			: FCoreStyle::Get().GetSlateColor("DefaultForeground");
 	}
 
-	// Callback for getting the text of the pick button.
+	/** Callback for getting the text of the pick button. */
 	FString HandlePickButtonText() const;
 
-	// Callback for generating a row in the reflector tree view.
+	/** Callback for generating a row in the reflector tree view. */
 	TSharedRef<ITableRow> HandleReflectorTreeGenerateRow( TSharedPtr<FReflectorNode> InReflectorNode, const TSharedRef<STableViewBase>& OwnerTable );
 
-	// Callback for getting the child items of the given reflector tree node.
+	/** Callback for getting the child items of the given reflector tree node. */
 	void HandleReflectorTreeGetChildren( TSharedPtr<FReflectorNode> InWidgetGeometry, TArray<TSharedPtr<FReflectorNode>>& OutChildren )
 	{
 		OutChildren = InWidgetGeometry->ChildNodes;
 	}
 
-	// Callback for when the selection in the reflector tree has changed.
+	/** Callback for when the selection in the reflector tree has changed. */
 	void HandleReflectorTreeSelectionChanged( TSharedPtr< FReflectorNode >, ESelectInfo::Type /*SelectInfo*/ )
 	{
 		SelectedNodes = ReflectorTree->GetSelectedItems();
diff --git a/Engine/Source/Developer/SlateReflector/Private/Widgets/SWidgetReflectorToolTipWidget.h b/Engine/Source/Developer/SlateReflector/Private/Widgets/SWidgetReflectorToolTipWidget.h
--- a/Engine/Source/Developer/SlateReflector/Private/Widgets/SWidgetReflectorToolTipWidget.h
+++ b/Engine/Source/Developer/SlateReflector/Private/Widgets/SWidgetReflectorToolTipWidget.h
@@ -95,7 +95,7 @@ public:
 
 private:
 
-	FString GetWidgetsDesiredSize( ) const
+	FString GetWidgetsDesiredSize() const
 	{
 		TSharedPtr<SWidget> TheWidget = WidgetInfo.Get()->Widget.Pin();
 
@@ -104,12 +104,12 @@ private:
 			: FString();
 	}
 
-	FString GetWidgetActualSize( ) const
+	FString GetWidgetActualSize() const
 	{
 		return WidgetInfo.Get()->Geometry.Size.ToString();
 	}
 
-	FString GetSizeInfo( ) const
+	FString GetSizeInfo() const
 	{
 		TSharedPtr<SWidget> TheWidget = WidgetInfo.Get()->Widget.Pin();
 
@@ -118,7 +118,7 @@ private:
 			: FString();
 	}
 
-	FString GetEnabled( ) const
+	FString GetEnabled() const
 	{
 		TSharedPtr<SWidget> TheWidget = WidgetInfo.Get()->Widget.Pin();
 
@@ -129,7 +129,7 @@ private:
 
 private:
 
-	// The info about the widget that we are visualizing.
+	/** The info about the widget that we are visualizing. */
 	TAttribute<TSharedPtr<FReflectorNode>> WidgetInfo;
 };
 
diff --git a/Engine/Source/Developer/SlateReflector/Private/Widgets/SWidgetReflectorTreeWidgetItem.cpp b/Engine/Source/Developer/SlateReflector/Private/Widgets/SWidgetReflectorTreeWidgetItem.cpp
--- a/Engine/Source/Developer/SlateReflector/Private/Widgets/SWidgetReflectorTreeWidgetItem.cpp
+++ b/Engine/Source/Developer/SlateReflector/Private/Widgets/SWidgetReflectorTreeWidgetItem.cpp
@@ -4,6 +4,10 @@
 #include "SWidgetReflectorTreeWidgetItem.h"
 #include "SHyperlink.h"
 
+
+/* SMultiColumnTableRow overrides
+ *****************************************************************************/
+
 BEGIN_SLATE_FUNCTION_BUILD_OPTIMIZATION
 TSharedRef<SWidget> SReflectorTreeWidgetItem::GenerateWidgetForColumn(const FName& ColumnName)
 {
@@ -11,7 +15,7 @@ TSharedRef<SWidget> SReflectorTreeWidgetItem::GenerateWidgetForColumn(const FNam
 	{
 		return SNew(SHorizontalBox)
 
-			+ SHorizontalBox::Slot()
+		+ SHorizontalBox::Slot()
 			.AutoWidth()
 			.VAlign(VAlign_Center)
 			[
@@ -24,8 +28,8 @@ TSharedRef<SWidget> SReflectorTreeWidgetItem::GenerateWidgetForColumn(const FNam
 			.VAlign(VAlign_Center)
 			[
 				SNew(STextBlock)
-				.Text(this, &SReflectorTreeWidgetItem::GetWidgetType)
-				.ColorAndOpacity(this, &SReflectorTreeWidgetItem::GetTint)
+					.Text(this, &SReflectorTreeWidgetItem::GetWidgetType)
+					.ColorAndOpacity(this, &SReflectorTreeWidgetItem::GetTint)
 			];
 	}
 	else if (ColumnName == TEXT("WidgetInfo"))
@@ -36,8 +40,8 @@ TSharedRef<SWidget> SReflectorTreeWidgetItem::GenerateWidgetForColumn(const FNam
 			.Padding(FMargin(2.0f, 0.0f))
 			[
 				SNew(SHyperlink)
-				.Text(this, &SReflectorTreeWidgetItem::GetReadableLocation)
-				.OnNavigate(this, &SReflectorTreeWidgetItem::HandleHyperlinkNavigate)
+					.Text(this, &SReflectorTreeWidgetItem::GetReadableLocation)
+					.OnNavigate(this, &SReflectorTreeWidgetItem::HandleHyperlinkNavigate)
 			];
 	}
 	else if (ColumnName == "Visibility")
@@ -48,7 +52,7 @@ TSharedRef<SWidget> SReflectorTreeWidgetItem::GenerateWidgetForColumn(const FNam
 			.Padding(FMargin(2.0f, 0.0f))
 			[
 				SNew(STextBlock)
-				.Text(this, &SReflectorTreeWidgetItem::GetVisibilityAsString)
+					.Text(this, &SReflectorTreeWidgetItem::GetVisibilityAsString)
 			];
 	}
 	else if (ColumnName == "ForegroundColor")
@@ -68,22 +72,21 @@ TSharedRef<SWidget> SReflectorTreeWidgetItem::GenerateWidgetForColumn(const FNam
 			[
 				// Show a color block
 				SNew(SColorBlock)
-				.Color(Foreground.GetSpecifiedColor())
-				.Size(FVector2D(16.0f, 16.0f))
+					.Color(Foreground.GetSpecifiedColor())
+					.Size(FVector2D(16.0f, 16.0f))
 			];
 	}
 	else if (ColumnName == "Address")
 	{
 		const TSharedPtr<SWidget> TheWidget = WidgetInfo.Get()->Widget.Pin();
 		const FText Address = (TheWidget.IsValid())
-			? FText::FromString( FString::Printf(TEXT("0x%08X"), TheWidget.Get()) ) 
+			? FText::FromString(FString::Printf(TEXT("0x%08X"), TheWidget.Get()))
 			: NSLOCTEXT("SWidgetReflector","nullptr","nullptr");
 			
-		return
-			SNew(SHyperlink)
+		return SNew(SHyperlink)
 			.ToolTipText(NSLOCTEXT("SWidgetReflector", "ClickToCopy", "Click to copy address."))
-			.Text( Address )
-			.OnNavigate_Lambda( [Address](){ FPlatformMisc::ClipboardCopy(*Address.ToString()); } ) ;
+			.Text(Address)
+			.OnNavigate_Lambda([Address](){ FPlatformMisc::ClipboardCopy(*Address.ToString()); }) ;
 	}
 	else
 	{
diff --git a/Engine/Source/Developer/SlateReflector/Private/Widgets/SWidgetReflectorTreeWidgetItem.h b/Engine/Source/Developer/SlateReflector/Private/Widgets/SWidgetReflectorTreeWidgetItem.h
--- a/Engine/Source/Developer/SlateReflector/Private/Widgets/SWidgetReflectorTreeWidgetItem.h
+++ b/Engine/Source/Developer/SlateReflector/Private/Widgets/SWidgetReflectorTreeWidgetItem.h
@@ -4,7 +4,7 @@
 
 
 /**
- * Widget that visualizes the contents of a FReflectorNode
+ * Widget that visualizes the contents of a FReflectorNode.
  */
 class SLATEREFLECTOR_API SReflectorTreeWidgetItem
 	: public SMultiColumnTableRow<TSharedPtr<FReflectorNode>>
@@ -36,6 +36,10 @@ public:
 		SMultiColumnTableRow< TSharedPtr<FReflectorNode> >::Construct( SMultiColumnTableRow< TSharedPtr<FReflectorNode> >::FArguments().Padding(1), InOwnerTableView );
 	}
 
+public:
+
+	// SMultiColumnTableRow overrides
+
 	virtual TSharedRef<SWidget> GenerateWidgetForColumn( const FName& ColumnName ) override;
 
 protected:
@@ -78,7 +82,7 @@ protected:
 	}
 
 	/** @return The tint of the reflector node */
-	FSlateColor GetTint () const
+	FSlateColor GetTint() const
 	{
 		return WidgetInfo.Get()->Tint;
 	}
@@ -91,7 +95,9 @@ protected:
 		}
 	}
 
-	/** The info about the widget that we are visualizing */
+private:
+
+	/** The info about the widget that we are visualizing. */
 	TAttribute< TSharedPtr<FReflectorNode> > WidgetInfo;
 
 	FAccessSourceCode OnAccessSourceCode;
diff --git a/Engine/Source/Developer/SlateReflector/Public/ISlateReflectorModule.h b/Engine/Source/Developer/SlateReflector/Public/ISlateReflectorModule.h
new file mode 100644
--- /dev/null
+++ b/Engine/Source/Developer/SlateReflector/Public/ISlateReflectorModule.h
@@ -0,0 +1,42 @@
+// Copyright 1998-2014 Epic Games, Inc. All Rights Reserved.
+
+#pragma once
+
+#include "ModuleInterface.h"
+
+
+// forward declarations
+class FWorkspaceItem;
+class SWidget;
+
+
+/**
+ * Interface for messaging modules.
+ */
+class ISlateReflectorModule
+	: public IModuleInterface
+{
+public:
+
+	/**
+	 * Creates a widget reflector widget.
+	 *
+	 * @return The widget.
+	 */
+	virtual TSharedRef<SWidget> GetWidgetReflector() = 0;
+
+	/**
+	 * Registers a tab spawner for the widget reflector.
+	 *
+	 * @param WorkspaceGroup The workspace group to insert the tab into.
+	 */
+	virtual void RegisterTabSpawner( const TSharedRef<FWorkspaceItem>& WorkspaceGroup ) = 0;
+
+	/** Unregisters the tab spawner for the widget reflector. */
+	virtual void UnregisterTabSpawner() = 0;
+
+public:
+
+	/** Virtual destructor. */
+	virtual ~ISlateReflectorModule() { }
+};
diff --git a/Engine/Source/Developer/SlateReflector/Public/Interfaces/ISlateReflectorModule.h b/Engine/Source/Developer/SlateReflector/Public/Interfaces/ISlateReflectorModule.h
deleted file mode 100644
diff --git a/Engine/Source/Developer/SlateReflector/Public/SlateReflector.h b/Engine/Source/Developer/SlateReflector/Public/SlateReflector.h
deleted file mode 100644
diff --git a/Engine/Source/Developer/SlateReflector/SlateReflector.Build.cs b/Engine/Source/Developer/SlateReflector/SlateReflector.Build.cs
--- a/Engine/Source/Developer/SlateReflector/SlateReflector.Build.cs
+++ b/Engine/Source/Developer/SlateReflector/SlateReflector.Build.cs
@@ -9,14 +9,15 @@ public class SlateReflector : ModuleRules
 		PublicDependencyModuleNames.AddRange(
 			new string[] {
 				"Core",
-				"Slate",
-				"SlateCore",
 			}
 		);
 
 		PrivateDependencyModuleNames.AddRange(
 			new string[] {
+				"CoreUObject",
 				"InputCore",
+				"Slate",
+				"SlateCore",
 			}
 		);
 
diff --git a/Engine/Source/Editor/LevelEditor/Private/LevelEditor.cpp b/Engine/Source/Editor/LevelEditor/Private/LevelEditor.cpp
--- a/Engine/Source/Editor/LevelEditor/Private/LevelEditor.cpp
+++ b/Engine/Source/Editor/LevelEditor/Private/LevelEditor.cpp
@@ -15,7 +15,7 @@
 #include "LevelViewportActions.h"
 #include "GlobalEditorCommonCommands.h"
 #include "IUserFeedbackModule.h"
-#include "SlateReflector.h"
+#include "ISlateReflectorModule.h"
 #include "SDockTab.h"
 #include "ToolkitManager.h"
 #include "TargetPlatform.h"
diff --git a/Engine/Source/Programs/CrashReporter/CrashReportClient/Private/CrashReportClientApp.cpp b/Engine/Source/Programs/CrashReporter/CrashReportClient/Private/CrashReportClientApp.cpp
--- a/Engine/Source/Programs/CrashReporter/CrashReportClient/Private/CrashReportClientApp.cpp
+++ b/Engine/Source/Programs/CrashReporter/CrashReportClient/Private/CrashReportClientApp.cpp
@@ -8,7 +8,7 @@
 	#include "SCrashReportClient.h"
 	#include "CrashReportClient.h"
 	#include "CrashReportClientStyle.h"
-	#include "SlateReflector.h"
+	#include "ISlateReflectorModule.h"
 #endif // !CRASH_REPORT_UNATTENDED_ONLY
 
 #include "CrashReportClientUnattended.h"
diff --git a/Engine/Source/Programs/SlateViewer/Public/SlateViewerApp.h b/Engine/Source/Programs/SlateViewer/Public/SlateViewerApp.h
--- a/Engine/Source/Programs/SlateViewer/Public/SlateViewerApp.h
+++ b/Engine/Source/Programs/SlateViewer/Public/SlateViewerApp.h
@@ -1,12 +1,12 @@
 // Copyright 1998-2014 Epic Games, Inc. All Rights Reserved.
 
-
 #pragma once
 
+#include "ISlateReflectorModule.h"
 #include "SlateViewer.h"
-#include "SlateReflector.h"
+
 
 /**
- * Run the SlateViewer 
+ * Run the SlateViewer .
  */
 void RunSlateViewer(const TCHAR* Commandline);
diff --git a/Engine/Source/Programs/UnrealFrontend/Private/Commands/UserInterfaceCommand.cpp b/Engine/Source/Programs/UnrealFrontend/Private/Commands/UserInterfaceCommand.cpp
--- a/Engine/Source/Programs/UnrealFrontend/Private/Commands/UserInterfaceCommand.cpp
+++ b/Engine/Source/Programs/UnrealFrontend/Private/Commands/UserInterfaceCommand.cpp
@@ -2,9 +2,9 @@
 
 #include "UnrealFrontendPrivatePCH.h"
 #include "AutomationController.h"
+#include "ISlateReflectorModule.h"
 #include "Projects.h"
 #include "SlateBasics.h"
-#include "SlateReflector.h"
 #include "StandaloneRenderer.h"
 #include "TaskGraphInterfaces.h"
 #include "UserInterfaceCommand.h"
@@ -35,7 +35,7 @@ void FUserInterfaceCommand::Run(  )
 	AutomationControllerModule.Init();
 
 	// load plug-ins
-	// @todo: allow for better plug-in support in standalone Slate apps
+	// @todo: allow for better plug-in support in standalone Slate applications
 	IPluginManager::Get().LoadModulesForEnabledPlugins(ELoadingPhase::PreDefault);
 
 	// load optional modules
diff --git a/Engine/Source/Runtime/Engine/Private/CheatManager.cpp b/Engine/Source/Runtime/Engine/Private/CheatManager.cpp
--- a/Engine/Source/Runtime/Engine/Private/CheatManager.cpp
+++ b/Engine/Source/Runtime/Engine/Private/CheatManager.cpp
@@ -3,7 +3,7 @@
 #include "EnginePrivate.h"
 #include "UnrealNetwork.h"
 #include "SlateBasics.h"
-#include "SlateReflector.h"
+#include "ISlateReflectorModule.h"
 #include "NavDataGenerator.h"
 #include "OnlineSubsystemUtils.h"
 #include "VisualLogger/VisualLogger.h"
