UMG - The Widget Reflector can detect UMG widgets as being different than normal slate widgets, and the source link is replaced by a link back to the blueprint asset.  This enables the functionality, another coming changelist will hookup the asset loading binding in the editor.

diff --git a/Engine/Source/Developer/SlateReflector/Private/Widgets/SWidgetReflector.cpp b/Engine/Source/Developer/SlateReflector/Private/Widgets/SWidgetReflector.cpp
--- a/Engine/Source/Developer/SlateReflector/Private/Widgets/SWidgetReflector.cpp
+++ b/Engine/Source/Developer/SlateReflector/Private/Widgets/SWidgetReflector.cpp
@@ -113,41 +113,41 @@ void SWidgetReflector::Construct( const FArguments& InArgs )
 					[
 						// The tree view that shows all the info that we capture.
 						SAssignNew(ReflectorTree, SReflectorTree)
-							.ItemHeight(24.0f)
-							.TreeItemsSource(&ReflectorTreeRoot)
-							.OnGenerateRow(this, &SWidgetReflector::HandleReflectorTreeGenerateRow)
-							.OnGetChildren(this, &SWidgetReflector::HandleReflectorTreeGetChildren)
-							.OnSelectionChanged(this, &SWidgetReflector::HandleReflectorTreeSelectionChanged)
-							.HeaderRow
-							(
-								SNew(SHeaderRow)
-
-								+ SHeaderRow::Column("WidgetName")
-									.DefaultLabel(LOCTEXT("WidgetName", "Widget Name").ToString())
-									.FillWidth(0.65f)
-
-								+ SHeaderRow::Column("ForegroundColor")
-									.FixedWidth(24.0f)
-									.VAlignHeader(VAlign_Center)
-									.HeaderContent()
-									[
-										SNew(STextBlock)
-											.Text(LOCTEXT("ForegroundColor", "FG").ToString())
-											.ToolTipText(LOCTEXT("ForegroundColorToolTip", "Foreground Color").ToString())
-									]
+						.ItemHeight(24.0f)
+						.TreeItemsSource(&ReflectorTreeRoot)
+						.OnGenerateRow(this, &SWidgetReflector::HandleReflectorTreeGenerateRow)
+						.OnGetChildren(this, &SWidgetReflector::HandleReflectorTreeGetChildren)
+						.OnSelectionChanged(this, &SWidgetReflector::HandleReflectorTreeSelectionChanged)
+						.HeaderRow
+						(
+							SNew(SHeaderRow)
+
+							+ SHeaderRow::Column("WidgetName")
+							.DefaultLabel(LOCTEXT("WidgetName", "Widget Name").ToString())
+							.FillWidth(0.65f)
+
+							+ SHeaderRow::Column("ForegroundColor")
+							.FixedWidth(24.0f)
+							.VAlignHeader(VAlign_Center)
+							.HeaderContent()
+							[
+								SNew(STextBlock)
+								.Text(LOCTEXT("ForegroundColor", "FG").ToString())
+								.ToolTipText(LOCTEXT("ForegroundColorToolTip", "Foreground Color").ToString())
+							]
 
-								+ SHeaderRow::Column("Visibility")
-									.DefaultLabel(LOCTEXT("Visibility", "Visibility" ).ToString())
-									.FixedWidth(125.0f)
+							+ SHeaderRow::Column("Visibility")
+							.DefaultLabel(LOCTEXT("Visibility", "Visibility" ).ToString())
+							.FixedWidth(125.0f)
 
-								+ SHeaderRow::Column("WidgetInfo")
-									.DefaultLabel(LOCTEXT("WidgetInfo", "Widget Info" ).ToString())
-									.FillWidth(0.25f)
+							+ SHeaderRow::Column("WidgetInfo")
+							.DefaultLabel(LOCTEXT("WidgetInfo", "Widget Info" ).ToString())
+							.FillWidth(0.25f)
 
-								+ SHeaderRow::Column("Address")
-									.DefaultLabel( LOCTEXT("Address", "Address") )
-									.FillWidth( 0.10f )
-							)
+							+ SHeaderRow::Column("Address")
+							.DefaultLabel( LOCTEXT("Address", "Address") )
+							.FillWidth( 0.10f )
+						)
 					]
 
 				#if WITH_EVENT_LOGGING
@@ -381,7 +381,14 @@ TSharedRef<ITableRow> SWidgetReflector::HandleReflectorTreeGenerateRow( TSharedP
 	return SNew(SReflectorTreeWidgetItem, OwnerTable)
 		.WidgetInfoToVisualize(InReflectorNode)
 		.ToolTip(GenerateToolTipForReflectorNode(InReflectorNode))
-		.SourceCodeAccessor(SourceAccessDelegate);
+		.SourceCodeAccessor(SourceAccessDelegate)
+		.AssetAccessor(AsseetAccessDelegate);
+}
+
+
+void SWidgetReflector::HandleReflectorTreeGetChildren(TSharedPtr<FReflectorNode> InWidgetGeometry, TArray<TSharedPtr<FReflectorNode>>& OutChildren)
+{
+	OutChildren = InWidgetGeometry->ChildNodes;
 }
 
 
@@ -390,7 +397,7 @@ TSharedRef<ITableRow> SWidgetReflector::GenerateEventLogRow( TSharedRef<FLoggedE
 	return SNew(STableRow<TSharedRef<FLoggedEvent>>, OwnerTable)
 	[
 		SNew(STextBlock)
-			.Text(InLoggedEvent->ToText())
+		.Text(InLoggedEvent->ToText())
 	];
 }
 
diff --git a/Engine/Source/Developer/SlateReflector/Private/Widgets/SWidgetReflector.h b/Engine/Source/Developer/SlateReflector/Private/Widgets/SWidgetReflector.h
--- a/Engine/Source/Developer/SlateReflector/Private/Widgets/SWidgetReflector.h
+++ b/Engine/Source/Developer/SlateReflector/Private/Widgets/SWidgetReflector.h
@@ -70,6 +70,11 @@ public:
 		SourceAccessDelegate = InDelegate;
 	}
 
+	virtual void SetAssetAccessDelegate(FAccessAsset InDelegate) override
+	{
+		AsseetAccessDelegate = InDelegate;
+	}
+
 	virtual void SetWidgetsToVisualize( const FWidgetPath& InWidgetsToVisualize ) override;
 	virtual int32 Visualize( const FWidgetPath& InWidgetsToVisualize, FSlateWindowElementList& OutDrawElements, int32 LayerId ) override;
 
@@ -165,10 +170,7 @@ private:
 	TSharedRef<ITableRow> HandleReflectorTreeGenerateRow( TSharedPtr<FReflectorNode> InReflectorNode, const TSharedRef<STableViewBase>& OwnerTable );
 
 	/** Callback for getting the child items of the given reflector tree node. */
-	void HandleReflectorTreeGetChildren( TSharedPtr<FReflectorNode> InWidgetGeometry, TArray<TSharedPtr<FReflectorNode>>& OutChildren )
-	{
-		OutChildren = InWidgetGeometry->ChildNodes;
-	}
+	void HandleReflectorTreeGetChildren( TSharedPtr<FReflectorNode> InWidgetGeometry, TArray<TSharedPtr<FReflectorNode>>& OutChildren );
 
 	/** Callback for when the selection in the reflector tree has changed. */
 	void HandleReflectorTreeSelectionChanged( TSharedPtr< FReflectorNode >, ESelectInfo::Type /*SelectInfo*/ )
@@ -192,6 +194,7 @@ private:
 	SSplitter::FSlot* WidgetInfoLocation;
 
 	FAccessSourceCode SourceAccessDelegate;
+	FAccessAsset AsseetAccessDelegate;
 
 	bool bShowFocus;
 	bool bIsPicking;
diff --git a/Engine/Source/Developer/SlateReflector/Private/Widgets/SWidgetReflectorTreeWidgetItem.cpp b/Engine/Source/Developer/SlateReflector/Private/Widgets/SWidgetReflectorTreeWidgetItem.cpp
--- a/Engine/Source/Developer/SlateReflector/Private/Widgets/SWidgetReflectorTreeWidgetItem.cpp
+++ b/Engine/Source/Developer/SlateReflector/Private/Widgets/SWidgetReflectorTreeWidgetItem.cpp
@@ -1,6 +1,7 @@
 // Copyright 1998-2014 Epic Games, Inc. All Rights Reserved.
 
 #include "SlateReflectorPrivatePCH.h"
+#include "ReflectionMetaData.h"
 #include "SWidgetReflectorTreeWidgetItem.h"
 #include "SHyperlink.h"
 
@@ -16,21 +17,21 @@ TSharedRef<SWidget> SReflectorTreeWidgetItem::GenerateWidgetForColumn(const FNam
 		return SNew(SHorizontalBox)
 
 		+ SHorizontalBox::Slot()
-			.AutoWidth()
-			.VAlign(VAlign_Center)
-			[
-				SNew(SExpanderArrow, SharedThis(this))
-			]
+		.AutoWidth()
+		.VAlign(VAlign_Center)
+		[
+			SNew(SExpanderArrow, SharedThis(this))
+		]
 
 		+ SHorizontalBox::Slot()
-			.AutoWidth()
-			.Padding(2.0f, 0.0f)
-			.VAlign(VAlign_Center)
-			[
-				SNew(STextBlock)
-					.Text(this, &SReflectorTreeWidgetItem::GetWidgetType)
-					.ColorAndOpacity(this, &SReflectorTreeWidgetItem::GetTint)
-			];
+		.AutoWidth()
+		.Padding(2.0f, 0.0f)
+		.VAlign(VAlign_Center)
+		[
+			SNew(STextBlock)
+			.Text(this, &SReflectorTreeWidgetItem::GetWidgetType)
+			.ColorAndOpacity(this, &SReflectorTreeWidgetItem::GetTint)
+		];
 	}
 	else if (ColumnName == TEXT("WidgetInfo"))
 	{
@@ -40,8 +41,8 @@ TSharedRef<SWidget> SReflectorTreeWidgetItem::GenerateWidgetForColumn(const FNam
 			.Padding(FMargin(2.0f, 0.0f))
 			[
 				SNew(SHyperlink)
-					.Text(this, &SReflectorTreeWidgetItem::GetReadableLocation)
-					.OnNavigate(this, &SReflectorTreeWidgetItem::HandleHyperlinkNavigate)
+				.Text(this, &SReflectorTreeWidgetItem::GetReadableLocation)
+				.OnNavigate(this, &SReflectorTreeWidgetItem::HandleHyperlinkNavigate)
 			];
 	}
 	else if (ColumnName == "Visibility")
@@ -94,3 +95,49 @@ TSharedRef<SWidget> SReflectorTreeWidgetItem::GenerateWidgetForColumn(const FNam
 	}
 }
 END_SLATE_FUNCTION_BUILD_OPTIMIZATION
+
+FString SReflectorTreeWidgetItem::GetReadableLocation() const
+{
+	FString ReadableLocation;
+
+	TWeakPtr<SWidget> Widget = WidgetInfo.Get()->Widget;
+	if ( Widget.IsValid() )
+	{
+		TSharedPtr<SWidget> SafeWidget = Widget.Pin();
+		TSharedPtr<FReflectionMetaData> MetaData = SafeWidget->GetMetaData<FReflectionMetaData>();
+		if ( MetaData.IsValid() )
+		{
+			ReadableLocation = MetaData->Asset->GetName() + TEXT(" [") + MetaData->Name.ToString() + TEXT("]");
+		}
+		else
+		{
+			ReadableLocation = SafeWidget->GetReadableLocation();
+		}
+	}
+
+	return ReadableLocation;
+}
+
+void SReflectorTreeWidgetItem::HandleHyperlinkNavigate()
+{
+	TWeakPtr<SWidget> Widget = WidgetInfo.Get()->Widget;
+	if ( Widget.IsValid() )
+	{
+		TSharedPtr<SWidget> SafeWidget = Widget.Pin();
+		TSharedPtr<FReflectionMetaData> MetaData = SafeWidget->GetMetaData<FReflectionMetaData>();
+		if ( MetaData.IsValid() )
+		{
+			if ( OnAccessAsset.IsBound() )
+			{
+				OnAccessAsset.Execute(MetaData->Asset.Get());
+			}
+		}
+		else
+		{
+			if ( OnAccessSourceCode.IsBound() )
+			{
+				OnAccessSourceCode.Execute(GetWidgetFile(), GetWidgetLineNumber(), 0);
+			}
+		}
+	}
+}
diff --git a/Engine/Source/Developer/SlateReflector/Private/Widgets/SWidgetReflectorTreeWidgetItem.h b/Engine/Source/Developer/SlateReflector/Private/Widgets/SWidgetReflectorTreeWidgetItem.h
--- a/Engine/Source/Developer/SlateReflector/Private/Widgets/SWidgetReflectorTreeWidgetItem.h
+++ b/Engine/Source/Developer/SlateReflector/Private/Widgets/SWidgetReflectorTreeWidgetItem.h
@@ -14,10 +14,12 @@ public:
 	SLATE_BEGIN_ARGS(SReflectorTreeWidgetItem)
 		: _WidgetInfoToVisualize()
 		, _SourceCodeAccessor()
+		, _AssetAccessor()
 	{ }
 
 		SLATE_ARGUMENT(TSharedPtr<FReflectorNode>, WidgetInfoToVisualize)
 		SLATE_ARGUMENT(FAccessSourceCode, SourceCodeAccessor)
+		SLATE_ARGUMENT(FAccessAsset, AssetAccessor)
 
 	SLATE_END_ARGS()
 
@@ -32,6 +34,7 @@ public:
 	{
 		this->WidgetInfo = InArgs._WidgetInfoToVisualize;
 		this->OnAccessSourceCode = InArgs._SourceCodeAccessor;
+		this->OnAccessAsset = InArgs._AssetAccessor;
 
 		SMultiColumnTableRow< TSharedPtr<FReflectorNode> >::Construct( SMultiColumnTableRow< TSharedPtr<FReflectorNode> >::FArguments().Padding(1), InOwnerTableView );
 	}
@@ -52,12 +55,7 @@ protected:
 			: TEXT("Null Widget");
 	}
 	
-	FString GetReadableLocation() const
-	{
-		return WidgetInfo.Get()->Widget.IsValid()
-			? WidgetInfo.Get()->Widget.Pin()->GetReadableLocation()
-			: FString();
-	}
+	FString GetReadableLocation() const;
 
 	FString GetWidgetFile() const
 	{
@@ -87,13 +85,7 @@ protected:
 		return WidgetInfo.Get()->Tint;
 	}
 
-	void HandleHyperlinkNavigate()
-	{
-		if(OnAccessSourceCode.IsBound())
-		{
-			OnAccessSourceCode.Execute(GetWidgetFile(), GetWidgetLineNumber(), 0);
-		}
-	}
+	void HandleHyperlinkNavigate();
 
 private:
 
@@ -101,4 +93,6 @@ private:
 	TAttribute< TSharedPtr<FReflectorNode> > WidgetInfo;
 
 	FAccessSourceCode OnAccessSourceCode;
+
+	FAccessAsset OnAccessAsset;
 };
diff --git a/Engine/Source/Runtime/Slate/Private/Framework/Application/SlateApplication.cpp b/Engine/Source/Runtime/Slate/Private/Framework/Application/SlateApplication.cpp
--- a/Engine/Source/Runtime/Slate/Private/Framework/Application/SlateApplication.cpp
+++ b/Engine/Source/Runtime/Slate/Private/Framework/Application/SlateApplication.cpp
@@ -1413,7 +1413,7 @@ TSharedRef< FGenericWindow > FSlateApplication::MakeWindow( TSharedRef<SWindow>
 	Definition->SupportsMinimize = InSlateWindow->HasMinimizeBox();
 	Definition->SupportsMaximize = InSlateWindow->HasMaximizeBox();
 
-    Definition->IsModalWindow = InSlateWindow->IsModalWindow();
+	Definition->IsModalWindow = InSlateWindow->IsModalWindow();
 	Definition->IsRegularWindow = InSlateWindow->IsRegularWindow();
 	Definition->HasSizingFrame = InSlateWindow->HasSizingFrame();
 	Definition->SizeWillChangeOften = InSlateWindow->SizeWillChangeOften();
@@ -5085,10 +5085,15 @@ TSharedRef<FSlateApplication> FSlateApplication::InitializeAsStandaloneApplicati
 
 void FSlateApplication::SetWidgetReflector(const TSharedRef<IWidgetReflector>& WidgetReflector)
 {
-	if (SourceCodeAccessDelegate.IsBound())
+	if ( SourceCodeAccessDelegate.IsBound() )
 	{
 		WidgetReflector->SetSourceAccessDelegate(SourceCodeAccessDelegate);
 	}
 
+	if ( AssetAccessDelegate.IsBound() )
+	{
+		WidgetReflector->SetAssetAccessDelegate(AssetAccessDelegate);
+	}
+
 	WidgetReflectorPtr = WidgetReflector;
 }
diff --git a/Engine/Source/Runtime/Slate/Public/Framework/Application/IWidgetReflector.h b/Engine/Source/Runtime/Slate/Public/Framework/Application/IWidgetReflector.h
--- a/Engine/Source/Runtime/Slate/Public/Framework/Application/IWidgetReflector.h
+++ b/Engine/Source/Runtime/Slate/Public/Framework/Application/IWidgetReflector.h
@@ -6,6 +6,9 @@
 /** A Delegate for passing along a string of a source code location to access */
 DECLARE_DELEGATE_RetVal_ThreeParams(bool, FAccessSourceCode, const FString& /*FileName*/, int32 /*InLineNumber*/, int32 /*InColumnNumber*/);
 
+/** A Delegate for an asset object to access */
+DECLARE_DELEGATE_RetVal_OneParam(bool, FAccessAsset, UObject* /*InAsset*/);
+
 
 /**
  * Interface for widget reflectors.
@@ -64,6 +67,11 @@ public:
 	 * @param InDelegate A delegate to access source code with.
 	 */
 	virtual void SetSourceAccessDelegate( FAccessSourceCode InDelegate ) = 0;
+
+	/**
+	* @param InDelegate A delegate to access assets with.
+	*/
+	virtual void SetAssetAccessDelegate(FAccessAsset InDelegate) = 0;
 	
 	/**
 	 * @param ThisWindow    Do we want to draw something for this window?
diff --git a/Engine/Source/Runtime/Slate/Public/Framework/Application/SlateApplication.h b/Engine/Source/Runtime/Slate/Public/Framework/Application/SlateApplication.h
--- a/Engine/Source/Runtime/Slate/Public/Framework/Application/SlateApplication.h
+++ b/Engine/Source/Runtime/Slate/Public/Framework/Application/SlateApplication.h
@@ -511,6 +511,12 @@ public:
 		QuerySourceCodeAccessDelegate = QueryAccessDelegate;
 	}
 
+	/** @param AccessDelegate The delegate to pass along to the widget reflector */
+	void SetWidgetReflectorAssetAccessDelegate(FAccessAsset AccessDelegate)
+	{
+		AssetAccessDelegate = AccessDelegate;
+	}
+
 	/** @param Scale  Sets the ratio SlateUnit / ScreenPixel */
 	void SetApplicationScale( float InScale ){ Scale = InScale; }
 
@@ -1312,12 +1318,15 @@ private:
 	/** A Widget that introspects the current UI hierarchy */
 	TWeakPtr<IWidgetReflector> WidgetReflectorPtr;
 
-	/** Delegate for accessing source code to pass to any widget inspectors. */
+	/** Delegate for accessing source code, to pass to any widget inspectors. */
 	FAccessSourceCode SourceCodeAccessDelegate;
 
 	/** Delegate for querying if source code access is available */
 	FQueryAccessSourceCode QuerySourceCodeAccessDelegate;
 
+	/** Delegate for accessing assets, to pass to any widget inspectors. */
+	FAccessAsset AssetAccessDelegate;
+
 	/** System for logging all relevant slate events to a log file */
 	TSharedPtr<class IEventLogger> EventLogger;
 
diff --git a/Engine/Source/Runtime/Slate/Public/Framework/SlateDelegates.h b/Engine/Source/Runtime/Slate/Public/Framework/SlateDelegates.h
--- a/Engine/Source/Runtime/Slate/Public/Framework/SlateDelegates.h
+++ b/Engine/Source/Runtime/Slate/Public/Framework/SlateDelegates.h
@@ -211,3 +211,5 @@ DECLARE_DELEGATE_RetVal_OneParam( FReply, FOnKeyEvent, const FKeyEvent& )
 /** A Delegate for passing along a string of a source code location to access */
 DECLARE_DELEGATE_RetVal_ThreeParams(bool, FAccessSourceCode, const FString& /*FileName*/, int32 /*InLineNumber*/, int32 /*InColumnNumber*/);
 
+/** A Delegate for an asset object to access */
+DECLARE_DELEGATE_RetVal_OneParam(bool, FAccessAsset, UObject* /*InAsset*/);
diff --git a/Engine/Source/Runtime/SlateCore/Public/Types/ReflectionMetadata.h b/Engine/Source/Runtime/SlateCore/Public/Types/ReflectionMetadata.h
new file mode 100644
--- /dev/null
+++ b/Engine/Source/Runtime/SlateCore/Public/Types/ReflectionMetadata.h
@@ -0,0 +1,32 @@
+// Copyright 1998-2014 Epic Games, Inc. All Rights Reserved.
+
+#pragma once
+
+#include "ISlateMetaData.h"
+
+/**
+ * Reflection meta-data that can be used by the widget reflector to determine
+ * additional information about slate widgets that are constructed by UObject
+ * classes for UMG.
+ */
+class FReflectionMetaData : public ISlateMetaData
+{
+public:
+	SLATE_METADATA_TYPE(FReflectionMetaData, ISlateMetaData)
+
+	FReflectionMetaData(FName InName, UClass* InClass, UObject* InAsset)
+		: Name(InName)
+		, Class(InClass)
+		, Asset(InAsset)
+	{
+	}
+
+	/** The name of the widget in the hierarchy */
+	FName Name;
+	
+	/** The class the constructed the slate widget. */
+	TWeakObjectPtr<UClass> Class;
+
+	/** The asset that owns the widget and is responsible for its specific existance. */
+	TWeakObjectPtr<UObject> Asset;
+};
diff --git a/Engine/Source/Runtime/UMG/Private/Components/Widget.cpp b/Engine/Source/Runtime/UMG/Private/Components/Widget.cpp
--- a/Engine/Source/Runtime/UMG/Private/Components/Widget.cpp
+++ b/Engine/Source/Runtime/UMG/Private/Components/Widget.cpp
@@ -2,6 +2,8 @@
 
 #include "UMGPrivatePCH.h"
 
+#include "ReflectionMetaData.h"
+
 /////////////////////////////////////////////////////
 // UWidget
 
@@ -478,6 +480,10 @@ void UWidget::SynchronizeProperties()
 
 		Navigation->UpdateMetaData(MetaData.ToSharedRef());
 	}
+
+#if WITH_EDITOR
+	SafeWidget->AddMetadata<FReflectionMetaData>(MakeShareable(new FReflectionMetaData(GetFName(), GetClass(), WidgetGeneratedBy)));
+#endif
 }
 
 bool UWidget::IsDesignTime() const
diff --git a/Engine/Source/Runtime/UMG/Private/WidgetBlueprintGeneratedClass.cpp b/Engine/Source/Runtime/UMG/Private/WidgetBlueprintGeneratedClass.cpp
--- a/Engine/Source/Runtime/UMG/Private/WidgetBlueprintGeneratedClass.cpp
+++ b/Engine/Source/Runtime/UMG/Private/WidgetBlueprintGeneratedClass.cpp
@@ -15,44 +15,14 @@ UWidgetBlueprintGeneratedClass::UWidgetBlueprintGeneratedClass(const FObjectInit
 {
 }
 
-void UWidgetBlueprintGeneratedClass::PostInitProperties()
-{
-	Super::PostInitProperties();
-
-	//// Create a widget tree if one doesn't already exist.
-	//if ( WidgetTree == NULL )
-	//{
-	//	WidgetTree = ConstructObject<UWidgetTree>(UWidgetTree::StaticClass(), this);
-	//}
-
-	////WidgetTree->SetFlags(RF_DefaultSubObject);
-	//WidgetTree->SetFlags(RF_Transactional);
-}
-
-void UWidgetBlueprintGeneratedClass::Link(FArchive& Ar, bool bRelinkExistingProperties)
-{
-	Super::Link(Ar, bRelinkExistingProperties);
-
-	// @TODO: Shouldn't be necessary to clear these, but currently the class gets linked twice during compilation
-	WidgetNodeProperties.Empty();
-
-	// Initialize derived members
-	//for ( TFieldIterator<UProperty> It(this); It; ++It )
-	//{
-	//	if ( UStructProperty* StructProp = Cast<UStructProperty>(*It) )
-	//	{
-	//		if ( StructProp->Struct->IsChildOf(FWidgetNode_Base::StaticStruct()) )
-	//		{
-	//			WidgetNodeProperties.Add(StructProp);
-	//		}
-	//	}
-	//}
-}
-
 void UWidgetBlueprintGeneratedClass::InitializeWidget(UUserWidget* UserWidget) const
 {
 	UWidgetTree* ClonedTree = DuplicateObject<UWidgetTree>( WidgetTree, UserWidget );
 
+#if WITH_EDITOR
+	UserWidget->WidgetGeneratedBy = ClassGeneratedBy;
+#endif
+
 	if ( ClonedTree )
 	{
 		ClonedTree->SetFlags(RF_Transactional);
@@ -83,6 +53,10 @@ void UWidgetBlueprintGeneratedClass::InitializeWidget(UUserWidget* UserWidget) c
 				return;
 			}
 
+#if WITH_EDITOR
+			Widget->WidgetGeneratedBy = ClassGeneratedBy;
+#endif
+
 			// TODO UMG Make this an FName
 			FString VariableName = Widget->GetName();
 
diff --git a/Engine/Source/Runtime/UMG/Public/Blueprint/WidgetBlueprintGeneratedClass.h b/Engine/Source/Runtime/UMG/Public/Blueprint/WidgetBlueprintGeneratedClass.h
--- a/Engine/Source/Runtime/UMG/Public/Blueprint/WidgetBlueprintGeneratedClass.h
+++ b/Engine/Source/Runtime/UMG/Public/Blueprint/WidgetBlueprintGeneratedClass.h
@@ -45,9 +45,9 @@ struct FDelegateRuntimeBinding
 
 
 /**
- * The widget blueprint generated class allows us to create blueprintable widgets for UMG at runtime.
+ * The widget blueprint generated class allows us to create blueprint-able widgets for UMG at runtime.
  * All WBPGC's are of UUserWidget classes, and they perform special post initialization using this class
- * to give themselves many of the same capabiltities as AActor blueprints, like dynamic delegate binding for
+ * to give themselves many of the same capabilities as AActor blueprints, like dynamic delegate binding for
  * widgets.
  */
 UCLASS()
@@ -55,6 +55,8 @@ class UMG_API UWidgetBlueprintGeneratedClass : public UBlueprintGeneratedClass
 {
 	GENERATED_UCLASS_BODY()
 
+public:
+
 	/** A tree of the widget templates to be created */
 	UPROPERTY()
 	class UWidgetTree* WidgetTree;
@@ -63,22 +65,17 @@ class UMG_API UWidgetBlueprintGeneratedClass : public UBlueprintGeneratedClass
 	TArray< FDelegateRuntimeBinding > Bindings;
 
 	UPROPERTY()
-	TArray<UWidgetAnimation*> Animations;
+	TArray< UWidgetAnimation* > Animations;
 
 	UPROPERTY()
 	TArray< FName > NamedSlots;
 
-	/** This is transient data calculated at link time. */
-	TArray<UStructProperty*> WidgetNodeProperties;
-
-	virtual void PostInitProperties() override;
-
-	virtual void Link(FArchive& Ar, bool bRelinkExistingProperties) override;
+public:
 
 	/**
 	 * This is the function that makes UMG work.  Once a user widget is constructed, it will post load
 	 * call into its generated class and ask to be initialized.  The class will perform all the delegate
-	 * binding and wiring nessesary to have the user's widget perform as desired.
+	 * binding and wiring necessary to have the user's widget perform as desired.
 	 */
 	void InitializeWidget(UUserWidget* UserWidget) const;
 };
\ No newline at end of file
diff --git a/Engine/Source/Runtime/UMG/Public/Components/Widget.h b/Engine/Source/Runtime/UMG/Public/Components/Widget.h
--- a/Engine/Source/Runtime/UMG/Public/Components/Widget.h
+++ b/Engine/Source/Runtime/UMG/Public/Components/Widget.h
@@ -125,8 +125,14 @@ public:
 	UPROPERTY()
 	bool bHiddenInDesigner;
 
+	/** Stores a reference to the asset responsible for this widgets construction. */
+	UPROPERTY(Transient)
+	UObject* WidgetGeneratedBy;
+
 #endif
 
+public:
+
 	/** */
 	UFUNCTION(BlueprintCallable, Category="Widget|Transform")
 	void SetRenderTransform(FWidgetTransform InTransform);
