Fix inconsistent case.

diff --git a/Engine/Source/Developer/SlateReflector/Private/Widgets/SWidgetReflectorTreeWidgetItem.cpp b/Engine/Source/Developer/SlateReflector/Private/Widgets/SWidgetReflectorTreeWidgetItem.cpp
--- a/Engine/Source/Developer/SlateReflector/Private/Widgets/SWidgetReflectorTreeWidgetItem.cpp
+++ b/Engine/Source/Developer/SlateReflector/Private/Widgets/SWidgetReflectorTreeWidgetItem.cpp
@@ -1,7 +1,7 @@
 // Copyright 1998-2014 Epic Games, Inc. All Rights Reserved.
 
 #include "SlateReflectorPrivatePCH.h"
-#include "ReflectionMetaData.h"
+#include "ReflectionMetadata.h"
 #include "SWidgetReflectorTreeWidgetItem.h"
 #include "SHyperlink.h"
 
diff --git a/Engine/Source/Runtime/UMG/Private/Components/Widget.cpp b/Engine/Source/Runtime/UMG/Private/Components/Widget.cpp
--- a/Engine/Source/Runtime/UMG/Private/Components/Widget.cpp
+++ b/Engine/Source/Runtime/UMG/Private/Components/Widget.cpp
@@ -2,7 +2,7 @@
 
 #include "UMGPrivatePCH.h"
 
-#include "ReflectionMetaData.h"
+#include "ReflectionMetadata.h"
 
 /////////////////////////////////////////////////////
 // UWidget
