Slate : Added stat visualization to the widget reflector. Slate can now be easily profiled in SlateViewer. Requires adding <bCompileWithStatsWithoutEngine>true</bCompileWithStatsWithoutEngine> to a local BuildConfiguration.xml.

diff --git a/Engine/Source/Developer/SlateReflector/Private/Widgets/SWidgetReflector.cpp b/Engine/Source/Developer/SlateReflector/Private/Widgets/SWidgetReflector.cpp
--- a/Engine/Source/Developer/SlateReflector/Private/Widgets/SWidgetReflector.cpp
+++ b/Engine/Source/Developer/SlateReflector/Private/Widgets/SWidgetReflector.cpp
@@ -2,6 +2,10 @@
 
 #include "SlateReflectorPrivatePCH.h"
 
+#if STATS
+#include "StatsData.h"
+#endif
+
 
 #define LOCTEXT_NAMESPACE "SWidgetReflector"
 #define WITH_EVENT_LOGGING 0
@@ -36,6 +40,17 @@ struct FLoggedEvent
 /* SWidgetReflector interface
  *****************************************************************************/
 
+static FReply ToggleSlateStats()
+{
+	#if STATS
+	FSelfRegisteringExec::StaticExec(nullptr, TEXT("stat slate"), *GLog);
+	#endif
+
+	return FReply::Handled();
+}
+
+
+
 void SWidgetReflector::Construct( const FArguments& InArgs )
 {
 	LoggedEvents.Reserve(MaxLoggedEvents);
@@ -63,15 +78,15 @@ void SWidgetReflector::Construct( const FArguments& InArgs )
 							]
 
 						+ SHorizontalBox::Slot()
-							.MaxWidth(250)
-							[
-								SNew(SSpinBox<float>)
-									.Value(this, &SWidgetReflector::HandleAppScaleSliderValue)
-									.MinValue(0.1f)
-									.MaxValue(3.0f)
-									.Delta(0.01f)
-									.OnValueChanged(this, &SWidgetReflector::HandleAppScaleSliderChanged)
-							]
+						.MaxWidth(250)
+						[
+							SNew(SSpinBox<float>)
+								.Value(this, &SWidgetReflector::HandleAppScaleSliderValue)
+								.MinValue(0.1f)
+								.MaxValue(3.0f)
+								.Delta(0.01f)
+								.OnValueChanged(this, &SWidgetReflector::HandleAppScaleSliderChanged)
+						]
 					]
 
 				+ SVerticalBox::Slot()
@@ -80,32 +95,55 @@ void SWidgetReflector::Construct( const FArguments& InArgs )
 						SNew(SHorizontalBox)
 
 						+ SHorizontalBox::Slot()
-							.AutoWidth()
-							.Padding(5.0f)
-							[
-								// Check box that controls LIVE MODE
-								SNew(SCheckBox)
-									.IsChecked(this, &SWidgetReflector::HandleFocusCheckBoxIsChecked)
-									.OnCheckStateChanged(this, &SWidgetReflector::HandleFocusCheckBoxCheckedStateChanged)
-									[
-										SNew(STextBlock)
-											.Text(LOCTEXT("ShowFocus", "Show Focus").ToString())
-									]
-							]
+						.AutoWidth()
+						.Padding(5.0f)
+						[
+							// Check box that controls LIVE MODE
+							SNew(SCheckBox)
+								.IsChecked(this, &SWidgetReflector::HandleFocusCheckBoxIsChecked)
+								.OnCheckStateChanged(this, &SWidgetReflector::HandleFocusCheckBoxCheckedStateChanged)
+								[
+									SNew(STextBlock)
+										.Text(LOCTEXT("ShowFocus", "Show Focus").ToString())
+								]
+						]
 
 						+ SHorizontalBox::Slot()
-							.AutoWidth()
-							.Padding(5.0f)
+						.AutoWidth()
+						.Padding(5.0f)
+						[
+							// Check box that controls PICKING A WIDGET TO INSPECT
+							SNew(SButton)
+							.OnClicked(this, &SWidgetReflector::HandlePickButtonClicked)
+							.ButtonColorAndOpacity(this, &SWidgetReflector::HandlePickButtonColorAndOpacity)
 							[
-								// Check box that controls PICKING A WIDGET TO INSPECT
-								SNew(SButton)
-									.OnClicked(this, &SWidgetReflector::HandlePickButtonClicked)
-									.ButtonColorAndOpacity(this, &SWidgetReflector::HandlePickButtonColorAndOpacity)
-									[
-										SNew(STextBlock)
-											.Text(this, &SWidgetReflector::HandlePickButtonText)
-									]
+								SNew(STextBlock)
+								.Text(this, &SWidgetReflector::HandlePickButtonText)
 							]
+						]
+
+						+SHorizontalBox::Slot()
+						.AutoWidth()
+						.Padding(5.0f)
+						[
+							SNew(SButton)
+							#if STATS
+							.Text( LOCTEXT("ToggleStats", "Toggle Stats") )
+							#else
+							.Text( LOCTEXT("ToggleStatsUnavailable", "Toggle Stats [see tooltip]") )
+							.ToolTip
+							(
+								SNew(SToolTip)
+								[
+									SNew(STextBlock)
+									.WrapTextAt(200.0f)
+									.Text( LOCTEXT("ToggleStatsUnavailableTooltip", "To enable STATS in standalone applications turn on <bCompileWithStatsWithoutEngine>false</bCompileWithStatsWithoutEngine> in BuildConfiguration.xml") )
+								]
+							)
+							.IsEnabled(false)
+							#endif
+							.OnClicked_Static( &ToggleSlateStats )
+						]
 					]
 
 				+ SVerticalBox::Slot()
@@ -161,23 +199,21 @@ void SWidgetReflector::Construct( const FArguments& InArgs )
 				#endif //WITH_EVENT_LOGGING
 			
 				+ SVerticalBox::Slot()
-					.AutoHeight()
-					[
-						// Frame rate counter
-						SNew(STextBlock)
-							.Text(this, &SWidgetReflector::HandleFrameRateText)
-					]
+				.AutoHeight()
+				[
+					MakeStatViewer()
+				]
+				
 			]
 	];
 }
 
-
 /* SCompoundWidget overrides
  *****************************************************************************/
 
 void SWidgetReflector::Tick( const FGeometry& AllottedGeometry, const double InCurrentTime, const float InDeltaTime )
 {
-	SCompoundWidget::Tick(AllottedGeometry, InCurrentTime, InDeltaTime);
+	UpdateStats();
 }
 
 
@@ -402,4 +438,234 @@ TSharedRef<ITableRow> SWidgetReflector::GenerateEventLogRow( TSharedRef<FLoggedE
 }
 
 
+//
+// STATS
+//
+
+/**
+ * Holds the stat data for visualation. StatItems have an immutable
+ * identity and a mutable value that changes from frame to frame.
+ */
+class FStatItem
+{
+public:
+	
+	FStatItem( FName InName, FText InDisplayText )
+	: StatName( InName )
+	, DisplayText( InDisplayText )
+	, InclusiveAvgMs( 0.0f )
+	{
+	}
+
+	FName GetStatName() const { return StatName; }
+	FText GetDisplayName() const { return DisplayText; }	
+	FText GetInclusiveAvgMsText() const { return InclusiveAvgMsText; }
+	float GetInclusiveAvgMs() const { return InclusiveAvgMs; }
+	void UpdateInclusiveAvgMs( float InMs ){ InclusiveAvgMs = InMs; InclusiveAvgMsText = FText::AsNumber(InMs); }
+
+private:
+
+	FName StatName;
+	FText DisplayText;
+	FText InclusiveAvgMsText;
+	float InclusiveAvgMs;
+	
+};
+
+static const FName ColumnId_StatName("StatName");
+static const FName ColumnId_InclusiveAvgMs("InclusiveAvgMs");
+static const FName ColumnId_InclusiveAvgMsGraph("InclusiveAvgMsGraph");
+
+TSharedRef<SWidget> SWidgetReflector::MakeStatViewer()
+{
+	return
+		SAssignNew(StatsBorder, SBorder)
+		.BorderImage( FCoreStyle::Get().GetBrush("NoBrush") )
+		[
+			// STATS LIST
+			SAssignNew( StatsList, SListView< TSharedRef< FStatItem > > )
+			.OnGenerateRow( this, &SWidgetReflector::GenerateStatRow )
+			.ListItemsSource( &StatsItems )
+			.HeaderRow
+			(
+				// STATS HEADER
+				SNew(SHeaderRow)
+				
+				+SHeaderRow::Column(ColumnId_StatName)
+				.FillWidth( 5.0f )
+				.HAlignCell(HAlign_Right)
+				.DefaultLabel( LOCTEXT("Stats_StatNameColumn", "Statistic") )
+
+				+SHeaderRow::Column(ColumnId_InclusiveAvgMs)
+				.FixedWidth(80.0f)
+				.HAlignCell(HAlign_Right)
+				.DefaultLabel( LOCTEXT("Stats_InclusiveAvgMsColumn", "AvgTime (ms)") )
+
+				+SHeaderRow::Column(ColumnId_InclusiveAvgMsGraph)
+				.FillWidth( 7.0f )
+				.DefaultLabel( LOCTEXT("Stats_InclusiveAvgMsGraphColumn", " ") )
+			)
+		];
+}
+
+void SWidgetReflector::UpdateStats()
+{
+	#if STATS
+
+	const FGameThreadHudData* ViewDataPtr = FHUDGroupGameThreadRenderer::Get().Latest;
+	if (ViewDataPtr != nullptr)
+	{
+		// Stats are up to date; make sure the UI is full bright.
+		StatsBorder->SetColorAndOpacity(FLinearColor::White);
+
+		// If we don't get a report about a stat this frame, assume its value is 0;
+		for (auto& StatItem : StatsItems)
+		{
+			StatItem->UpdateInclusiveAvgMs(0.0f);
+		}
+
+		// Update stats with latest values; add any stats if they are missing
+		bool bNeedsRefresh = false;
+		const FGameThreadHudData &ViewData = *ViewDataPtr;
+		for (int32 GroupIndex = 0; GroupIndex < ViewData.HudGroups.Num(); ++GroupIndex)
+		{
+			const FHudGroup& HudGroup = ViewData.HudGroups[GroupIndex];
+			const bool bHasHierarchy = !!HudGroup.HierAggregate.Num();
+			const bool bHasFlat = !!HudGroup.FlatAggregate.Num();
+
+			int32 Row = 0;
+			for (const FComplexStatMessage& StatItem : HudGroup.FlatAggregate)
+			{
+				const TSharedRef<FStatItem>* ExistingItemPtr = StatsItems.FindByPredicate( [&StatItem]( const TSharedRef<FStatItem>& ItemRef )
+				{
+					return ItemRef->GetStatName() == StatItem.GetShortName();
+				});
+
+				const float InMs = FPlatformTime::ToMilliseconds(StatItem.GetValue_Duration(EComplexStatField::IncAve));
+
+				// If stat does not exist...
+				if (ExistingItemPtr == nullptr)
+				{
+					// ... add it.
+					const TSharedRef<FStatItem> NewStatItem( new FStatItem( StatItem.GetShortName(), FText::FromString(StatItem.GetDescription() ) ) );
+					StatsItems.Add( NewStatItem );
+					ExistingItemPtr = &StatsItems.Last();
+
+					// We have added items to the list, and it needs to be refreshed.
+					bNeedsRefresh = true;
+				}
+
+				// Update stat values.
+				(*ExistingItemPtr)->UpdateInclusiveAvgMs( InMs );
+			}
+		}
+
+		// If we add widgets in tick, we must refresh the list.
+		if (bNeedsRefresh)
+		{
+			StatsList->RequestListRefresh();
+		}		
+	}
+	else
+	{
+		// We are now showing stale stats; provide visual hint.
+		StatsBorder->SetColorAndOpacity(FLinearColor(1,1,1,0.75f));
+	}
+
+	#endif
+}
+
+TSharedRef<ITableRow> SWidgetReflector::GenerateStatRow( TSharedRef<FStatItem> StatItem, const TSharedRef<STableViewBase>& OwnerTable )
+{
+	class SStatTableRow : public SMultiColumnTableRow< TSharedRef<FStatItem> >
+	{
+	
+	public:
+		typedef SMultiColumnTableRow< TSharedRef<FStatItem> > Super;
+		
+		SLATE_BEGIN_ARGS( SStatTableRow )
+		{}
+		SLATE_END_ARGS();
+
+		void Construct( const FArguments& InArgs, const TSharedRef<STableViewBase>& OwnerTable, const TSharedRef<FStatItem>& InStatItem )
+		{
+			StatItem = InStatItem;
+			FSuperRowType::Construct( FTableRowArgs(), OwnerTable );
+		}
+
+		virtual TSharedRef<SWidget> GenerateWidgetForColumn( const FName& ColumnName ) override
+		{
+			if (ColumnName == ColumnId_StatName)
+			{
+				// STAT NAME
+				return
+				SNew(STextBlock)
+				.Text( StatItem->GetDisplayName() );
+			}
+			else if (ColumnName == ColumnId_InclusiveAvgMs)
+			{
+				// STAT NUMBER
+				return
+				SNew(SBox)
+				.Padding(FMargin(5,0))
+				[
+					SNew(STextBlock)
+					.Text( this, &SStatTableRow::GetInclusiveAvgMsText )
+				];
+			}
+			else if ( ColumnName == ColumnId_InclusiveAvgMsGraph )
+			{
+				// BAR GRAPH
+				return
+				SNew(SHorizontalBox)
+				+SHorizontalBox::Slot()
+				.Padding(0,1)
+				.FillWidth( TAttribute<float>(SharedThis(this), &SStatTableRow::BarGraph_GetFilledWidth) )
+				[
+					SNew(SImage)
+					.Image( FCoreStyle::Get().GetBrush("WhiteBrush") )
+					.ColorAndOpacity( this, &SStatTableRow::BarGraph_GetBarColor )
+				]
+				+SHorizontalBox::Slot()
+				.Padding(0,1)
+				.FillWidth( TAttribute<float>(SharedThis(this), &SStatTableRow::BarGraph_GetEmptyWidth) );
+			}
+			else
+			{
+				return SNew(SSpacer);
+			}
+		}
+
+		private:
+
+		FText GetInclusiveAvgMsText() const
+		{
+			return StatItem->GetInclusiveAvgMsText();
+		}
+
+		float BarGraph_GetFilledWidth() const
+		{
+			return StatItem->GetInclusiveAvgMs();
+		}
+
+		float BarGraph_GetEmptyWidth()  const
+		{
+			static const float MaxMs = 60.0f;
+			return MaxMs - StatItem->GetInclusiveAvgMs();
+		}
+
+		FSlateColor BarGraph_GetBarColor() const
+		{
+			static const float RedlineMs = 30.0f;
+			return FMath::Lerp( FLinearColor::Green, FLinearColor::Red, StatItem->GetInclusiveAvgMs()/RedlineMs );
+		}
+
+		TSharedPtr<FStatItem> StatItem;
+	};
+
+	return SNew(SStatTableRow, OwnerTable, StatItem);
+}
+
 #undef LOCTEXT_NAMESPACE
+
+
diff --git a/Engine/Source/Developer/SlateReflector/Private/Widgets/SWidgetReflector.h b/Engine/Source/Developer/SlateReflector/Private/Widgets/SWidgetReflector.h
--- a/Engine/Source/Developer/SlateReflector/Private/Widgets/SWidgetReflector.h
+++ b/Engine/Source/Developer/SlateReflector/Private/Widgets/SWidgetReflector.h
@@ -198,4 +198,14 @@ private:
 
 	bool bShowFocus;
 	bool bIsPicking;
+
+private:
+	// STATS
+	TSharedPtr<SBorder> StatsBorder;
+	TArray< TSharedRef<class FStatItem> > StatsItems;
+	TSharedPtr< SListView< TSharedRef< FStatItem > > > StatsList;
+
+	TSharedRef<SWidget> MakeStatViewer();
+	void UpdateStats();
+	TSharedRef<ITableRow> GenerateStatRow( TSharedRef<FStatItem> StatItem, const TSharedRef<STableViewBase>& OwnerTable );
 };
diff --git a/Engine/Source/Programs/SlateViewer/Private/SlateViewerApp.cpp b/Engine/Source/Programs/SlateViewer/Private/SlateViewerApp.cpp
--- a/Engine/Source/Programs/SlateViewer/Private/SlateViewerApp.cpp
+++ b/Engine/Source/Programs/SlateViewer/Private/SlateViewerApp.cpp
@@ -49,9 +49,11 @@ void RunSlateViewer( const TCHAR* CommandLine )
 	// loop while the server does the rest
 	while (!GIsRequestingExit)
 	{
-		FTicker::GetCoreTicker().Tick(0);
+		FTaskGraphInterface::Get().ProcessThreadUntilIdle(ENamedThreads::GameThread);
+		FStats::AdvanceFrame(false);
+		FTicker::GetCoreTicker().Tick(FApp::GetDeltaTime());
 		FSlateApplication::Get().PumpMessages();
-		FSlateApplication::Get().Tick();
+		FSlateApplication::Get().Tick();		
 		FPlatformProcess::Sleep(0);
 	}
 
