Slate : fixed Slate Draw Prepass only being called on top-level windows (not their children). Slate.FoldTick can now be toggle from the widget reflector.

diff --git a/Engine/Source/Developer/SlateReflector/Private/Widgets/SWidgetReflector.cpp b/Engine/Source/Developer/SlateReflector/Private/Widgets/SWidgetReflector.cpp
--- a/Engine/Source/Developer/SlateReflector/Private/Widgets/SWidgetReflector.cpp
+++ b/Engine/Source/Developer/SlateReflector/Private/Widgets/SWidgetReflector.cpp
@@ -49,6 +49,7 @@ static FReply ToggleSlateStats()
 	return FReply::Handled();
 }
 
+extern SLATECORE_API int32 bFoldTick;
 
 
 void SWidgetReflector::Construct( const FArguments& InArgs )
@@ -90,61 +91,85 @@ void SWidgetReflector::Construct( const FArguments& InArgs )
 					]
 
 				+ SVerticalBox::Slot()
-					.AutoHeight()
+				.AutoHeight()
+				[
+					SNew(SHorizontalBox)
+
+					+ SHorizontalBox::Slot()
+					.AutoWidth()
+					.Padding(5.0f)
 					[
-						SNew(SHorizontalBox)
+						// Check box that controls LIVE MODE
+						SNew(SCheckBox)
+							.IsChecked(this, &SWidgetReflector::HandleFocusCheckBoxIsChecked)
+							.OnCheckStateChanged(this, &SWidgetReflector::HandleFocusCheckBoxCheckedStateChanged)
+							[
+								SNew(STextBlock)
+									.Text(LOCTEXT("ShowFocus", "Show Focus").ToString())
+							]
+					]
 
-						+ SHorizontalBox::Slot()
-						.AutoWidth()
-						.Padding(5.0f)
+					+ SHorizontalBox::Slot()
+					.AutoWidth()
+					.Padding(5.0f)
+					[
+						// Check box that controls PICKING A WIDGET TO INSPECT
+						SNew(SButton)
+						.OnClicked(this, &SWidgetReflector::HandlePickButtonClicked)
+						.ButtonColorAndOpacity(this, &SWidgetReflector::HandlePickButtonColorAndOpacity)
 						[
-							// Check box that controls LIVE MODE
-							SNew(SCheckBox)
-								.IsChecked(this, &SWidgetReflector::HandleFocusCheckBoxIsChecked)
-								.OnCheckStateChanged(this, &SWidgetReflector::HandleFocusCheckBoxCheckedStateChanged)
-								[
-									SNew(STextBlock)
-										.Text(LOCTEXT("ShowFocus", "Show Focus").ToString())
-								]
+							SNew(STextBlock)
+							.Text(this, &SWidgetReflector::HandlePickButtonText)
 						]
+					]
 
-						+ SHorizontalBox::Slot()
-						.AutoWidth()
-						.Padding(5.0f)
-						[
-							// Check box that controls PICKING A WIDGET TO INSPECT
-							SNew(SButton)
-							.OnClicked(this, &SWidgetReflector::HandlePickButtonClicked)
-							.ButtonColorAndOpacity(this, &SWidgetReflector::HandlePickButtonColorAndOpacity)
+					+SHorizontalBox::Slot()
+					.AutoWidth()
+					.Padding(5.0f)
+					[
+						SNew(SButton)
+						#if STATS
+						.Text( LOCTEXT("ToggleStats", "Toggle Stats") )
+						#else
+						.Text( LOCTEXT("ToggleStatsUnavailable", "Toggle Stats [see tooltip]") )
+						.ToolTip
+						(
+							SNew(SToolTip)
 							[
 								SNew(STextBlock)
-								.Text(this, &SWidgetReflector::HandlePickButtonText)
+								.WrapTextAt(200.0f)
+								.Text( LOCTEXT("ToggleStatsUnavailableTooltip", "To enable STATS in standalone applications turn on <bCompileWithStatsWithoutEngine>false</bCompileWithStatsWithoutEngine> in BuildConfiguration.xml") )
 							]
-						]
-
-						+SHorizontalBox::Slot()
-						.AutoWidth()
-						.Padding(5.0f)
+						)
+						.IsEnabled(false)
+						#endif
+						.OnClicked_Static( &ToggleSlateStats )
+					]
+					+SHorizontalBox::Slot()
+					.AutoWidth()
+					.Padding(5.0f)
+					[
+						SNew(SCheckBox)
+						.Style( FCoreStyle::Get(), "ToggleButtonCheckbox" )
+						.IsChecked_Lambda([]()
+						{
+							return bFoldTick == 0 ? ESlateCheckBoxState::Unchecked : ESlateCheckBoxState::Checked;
+						})
+						.OnCheckStateChanged_Lambda([]( const ESlateCheckBoxState::Type NewState )
+						{
+							bFoldTick = (NewState == ESlateCheckBoxState::Checked) ? 1 : 0;
+						})
 						[
-							SNew(SButton)
-							#if STATS
-							.Text( LOCTEXT("ToggleStats", "Toggle Stats") )
-							#else
-							.Text( LOCTEXT("ToggleStatsUnavailable", "Toggle Stats [see tooltip]") )
-							.ToolTip
-							(
-								SNew(SToolTip)
-								[
-									SNew(STextBlock)
-									.WrapTextAt(200.0f)
-									.Text( LOCTEXT("ToggleStatsUnavailableTooltip", "To enable STATS in standalone applications turn on <bCompileWithStatsWithoutEngine>false</bCompileWithStatsWithoutEngine> in BuildConfiguration.xml") )
-								]
-							)
-							.IsEnabled(false)
-							#endif
-							.OnClicked_Static( &ToggleSlateStats )
+							SNew(SBox)
+							.VAlign( VAlign_Center )
+							.HAlign( HAlign_Center )
+							[
+								SNew(STextBlock)	
+								.Text( LOCTEXT("ToggleTickFolding", "Fold Tick") )
+							]
 						]
 					]
+				]
 
 				+ SVerticalBox::Slot()
 					.FillHeight(1.0f)
@@ -258,7 +283,7 @@ void SWidgetReflector::SetWidgetsToVisualize( const FWidgetPath& InWidgetsToVisu
 
 int32 SWidgetReflector::Visualize( const FWidgetPath& InWidgetsToVisualize, FSlateWindowElementList& OutDrawElements, int32 LayerId )
 {
-	const bool bAttemptingToVisualizeReflector = InWidgetsToVisualize.ContainsWidget(SharedThis(this));
+	const bool bAttemptingToVisualizeReflector = InWidgetsToVisualize.ContainsWidget(ReflectorTree.ToSharedRef());
 
 	if (!InWidgetsToVisualize.IsValid())
 	{
@@ -459,7 +484,10 @@ public:
 
 	FName GetStatName() const { return StatName; }
 	FText GetDisplayName() const { return DisplayText; }	
-	FText GetInclusiveAvgMsText() const { return InclusiveAvgMsText; }
+	FText GetInclusiveAvgMsText() const
+	{
+		return InclusiveAvgMsText;
+	}
 	float GetInclusiveAvgMs() const { return InclusiveAvgMs; }
 	void UpdateInclusiveAvgMs( float InMs )
 	{
@@ -474,8 +502,9 @@ public:
 
 private:
 
-	FName StatName;
-	FText DisplayText;
+	const FName StatName;
+	const FText DisplayText;
+
 	FText InclusiveAvgMsText;
 	float InclusiveAvgMs;
 	
@@ -574,11 +603,22 @@ void SWidgetReflector::UpdateStats()
 		{
 			StatsList->RequestListRefresh();
 		}		
+		else
+		{
+			int32 a = 5;
+			if (a == 10){  }
+		}
 	}
 	else
 	{
 		// We are now showing stale stats; provide visual hint.
 		StatsBorder->SetColorAndOpacity(FLinearColor(1,1,1,0.75f));
+		
+		// If we don't get a report about a stat this frame, assume its value is 0;
+		for (auto& StatItem : StatsItems)
+		{
+			StatItem->UpdateInclusiveAvgMs(0.0f);
+		}
 	}
 
 	#endif
diff --git a/Engine/Source/Runtime/Slate/Private/Framework/Application/SlateApplication.cpp b/Engine/Source/Runtime/Slate/Private/Framework/Application/SlateApplication.cpp
--- a/Engine/Source/Runtime/Slate/Private/Framework/Application/SlateApplication.cpp
+++ b/Engine/Source/Runtime/Slate/Private/Framework/Application/SlateApplication.cpp
@@ -952,18 +952,23 @@ void FSlateApplication::DrawWindowAndChildren( const TSharedRef<SWindow>& Window
 	}
 }
 
-static void DrawWindowPrepass( TSharedRef<SWindow> WindowToDraw )
+static void PrepassWindowAndChildren( TSharedRef<SWindow> WindowToPrepass )
 {
-	FScopedSwitchWorldHack SwitchWorld( WindowToDraw );
+	FScopedSwitchWorldHack SwitchWorld( WindowToPrepass );
 
 	{
 		SCOPE_CYCLE_COUNTER(STAT_SlateCacheDesiredSize);
-		WindowToDraw->SlatePrepass();
+		WindowToPrepass->SlatePrepass();
 	}
 
-	if (WindowToDraw->IsAutosized())
+	if (WindowToPrepass->IsAutosized())
 	{
-		WindowToDraw->ReshapeWindow(WindowToDraw->GetPositionInScreen(), WindowToDraw->GetDesiredSize());
+		WindowToPrepass->ReshapeWindow(WindowToPrepass->GetPositionInScreen(), WindowToPrepass->GetDesiredSize());
+	}
+
+	for (const TSharedRef<SWindow>& ChildWindow : WindowToPrepass->GetChildWindows() )
+	{
+		PrepassWindowAndChildren(ChildWindow);
 	}
 }
 
@@ -973,14 +978,14 @@ void FSlateApplication::DrawPrepass( TSharedPtr<SWindow> DrawOnlyThisWindow )
 
 	if (ActiveModalWindow.IsValid())
 	{
-		DrawWindowPrepass( ActiveModalWindow.ToSharedRef() );
+		PrepassWindowAndChildren( ActiveModalWindow.ToSharedRef() );
 
 		for (TArray< TSharedRef<SWindow> >::TConstIterator CurrentWindowIt(SlateWindows); CurrentWindowIt; ++CurrentWindowIt)
 		{
 			const TSharedRef<SWindow>& CurrentWindow = *CurrentWindowIt;
 			if (CurrentWindow->IsTopmostWindow())
 			{
-				DrawWindowPrepass( CurrentWindow );
+				PrepassWindowAndChildren( CurrentWindow );
 			}
 		}
 
@@ -988,22 +993,21 @@ void FSlateApplication::DrawPrepass( TSharedPtr<SWindow> DrawOnlyThisWindow )
 		FSlateNotificationManager::Get().GetWindows(NotificationWindows);
 		for (auto CurrentWindowIt(NotificationWindows.CreateIterator()); CurrentWindowIt; ++CurrentWindowIt)
 		{
-			DrawWindowPrepass(*CurrentWindowIt );
+			PrepassWindowAndChildren(*CurrentWindowIt );
 		}
 	}
 	else if (DrawOnlyThisWindow.IsValid())
 	{
-		DrawWindowPrepass( DrawOnlyThisWindow.ToSharedRef() );
+		PrepassWindowAndChildren( DrawOnlyThisWindow.ToSharedRef() );
 	}
 	else
 	{
 		// Draw all windows
-		for (TArray< TSharedRef<SWindow> >::TConstIterator CurrentWindowIt(SlateWindows); CurrentWindowIt; ++CurrentWindowIt)
+		for (const TSharedRef<SWindow>& CurrentWindow : SlateWindows)
 		{
-			TSharedRef<SWindow> CurrentWindow = *CurrentWindowIt;
 			if (CurrentWindow->IsVisible() && !CurrentWindow->IsWindowMinimized())
 			{
-				DrawWindowPrepass(CurrentWindow );
+				PrepassWindowAndChildren(CurrentWindow );
 			}
 		}
 	}
@@ -1144,7 +1148,7 @@ void FSlateApplication::FinishedInputThisFrame()
 	}
 }
 
-extern SLATECORE_API TAutoConsoleVariable<int32> FoldTick;
+extern SLATECORE_API int32 bFoldTick;
 
 /**
  * Ticks this application
@@ -1230,43 +1234,41 @@ void FSlateApplication::Tick()
 	// Update auto-throttling based on elapsed time since user interaction
 	ThrottleApplicationBasedOnMouseMovement();
 
-	const bool bFoldTick = (FoldTick.GetValueOnGameThread() != 0);
-
-	if ( !bFoldTick )
+	
+	if (!bFoldTick)
 	{
 		TSharedPtr<SWindow> ActiveModalWindow = GetActiveModalWindow();
-		{
-			SCOPE_CYCLE_COUNTER( STAT_SlateTickWindowAndChildren );
+	
+		SCOPE_CYCLE_COUNTER( STAT_SlateTickWindowAndChildren );
 
-			if ( ActiveModalWindow.IsValid() )
+		if ( ActiveModalWindow.IsValid() )
+		{
+			// There is a modal window, and we just need to tick it.
+			TickWindowAndChildren( ActiveModalWindow.ToSharedRef() );
+			// And also tick any top-level windows.
+			for( TArray< TSharedRef<SWindow> >::TIterator CurrentWindowIt( SlateWindows ); CurrentWindowIt; ++CurrentWindowIt )
 			{
-				// There is a modal window, and we just need to tick it.
-				TickWindowAndChildren( ActiveModalWindow.ToSharedRef() );
-				// And also tick any top-level windows.
-				for( TArray< TSharedRef<SWindow> >::TIterator CurrentWindowIt( SlateWindows ); CurrentWindowIt; ++CurrentWindowIt )
+				TSharedRef<SWindow>& CurrentWindow = *CurrentWindowIt;
+				if (CurrentWindow->IsTopmostWindow())
 				{
-					TSharedRef<SWindow>& CurrentWindow = *CurrentWindowIt;
-					if (CurrentWindow->IsTopmostWindow())
-					{
-						TickWindowAndChildren(CurrentWindow);
-					}
+					TickWindowAndChildren(CurrentWindow);
 				}
-				// also tick the notification manager's windows
-				TArray< TSharedRef<SWindow> > NotificationWindows;
-				FSlateNotificationManager::Get().GetWindows(NotificationWindows);
-				for( auto CurrentWindowIt( NotificationWindows.CreateIterator() ); CurrentWindowIt; ++CurrentWindowIt )
-				{
-					TickWindowAndChildren(*CurrentWindowIt);
-				}		
 			}
-			else
+			// also tick the notification manager's windows
+			TArray< TSharedRef<SWindow> > NotificationWindows;
+			FSlateNotificationManager::Get().GetWindows(NotificationWindows);
+			for( auto CurrentWindowIt( NotificationWindows.CreateIterator() ); CurrentWindowIt; ++CurrentWindowIt )
 			{
-				// No modal window; tick all slate windows.
-				for( TArray< TSharedRef<SWindow> >::TIterator CurrentWindowIt( SlateWindows ); CurrentWindowIt; ++CurrentWindowIt )
-				{
-					TSharedRef<SWindow>& CurrentWindow = *CurrentWindowIt;
-					TickWindowAndChildren( CurrentWindow );
-				}
+				TickWindowAndChildren(*CurrentWindowIt);
+			}		
+		}
+		else
+		{
+			// No modal window; tick all slate windows.
+			for( TArray< TSharedRef<SWindow> >::TIterator CurrentWindowIt( SlateWindows ); CurrentWindowIt; ++CurrentWindowIt )
+			{
+				TSharedRef<SWindow>& CurrentWindow = *CurrentWindowIt;
+				TickWindowAndChildren( CurrentWindow );
 			}
 		}
 	}
diff --git a/Engine/Source/Runtime/SlateCore/Private/Widgets/SWidget.cpp b/Engine/Source/Runtime/SlateCore/Private/Widgets/SWidget.cpp
--- a/Engine/Source/Runtime/SlateCore/Private/Widgets/SWidget.cpp
+++ b/Engine/Source/Runtime/SlateCore/Private/Widgets/SWidget.cpp
@@ -586,14 +586,14 @@ void SWidget::SetDebugInfo( const ANSICHAR* InType, const ANSICHAR* InFile, int3
 	this->CreatedOnLine = OnLine;
 }
 
-SLATECORE_API TAutoConsoleVariable<int32> FoldTick(TEXT("Slate.FoldTick"), 0, TEXT("When folding, call Tick as part of the paint pass instead of a separate tick pass."));
+SLATECORE_API int32 bFoldTick = 0;
+FAutoConsoleVariableRef FoldTick(TEXT("Slate.FoldTick"), bFoldTick, TEXT("When folding, call Tick as part of the paint pass instead of a separate tick pass."));
 
 int32 SWidget::Paint(const FPaintArgs& Args, const FGeometry& AllottedGeometry, const FSlateRect& MyClippingRect, FSlateWindowElementList& OutDrawElements, int32 LayerId, const FWidgetStyle& InWidgetStyle, bool bParentEnabled) const
 {
 	SCOPE_CYCLE_COUNTER(STAT_SlateOnPaint);
 	INC_DWORD_STAT(STAT_SlateNumPaintedWidgets);
 
-	const bool bFoldTick = (FoldTick.GetValueOnGameThread() != 0);
 	if ( bFoldTick )
 	{
 		FGeometry TickGeometry = AllottedGeometry;
@@ -605,7 +605,6 @@ int32 SWidget::Paint(const FPaintArgs& Args, const FGeometry& AllottedGeometry,
 
 	const FPaintArgs UpdatedArgs = Args.RecordHittestGeometry( this, AllottedGeometry, MyClippingRect );
 	return OnPaint(UpdatedArgs, AllottedGeometry, MyClippingRect, OutDrawElements, LayerId, InWidgetStyle, bParentEnabled);
-	//return OnPaint(Args, AllottedGeometry, MyClippingRect, OutDrawElements, LayerId, InWidgetStyle, bParentEnabled);
 }
 
 void SWidget::ArrangeChildren(const FGeometry& AllottedGeometry, FArrangedChildren& ArrangedChildren) const
