Slate : Added -perftest option to SlateViewer. Added button in WidgetReflector that copies Slate Perf Stats to clipboard.

diff --git a/Engine/Source/Developer/SlateReflector/Private/Widgets/SWidgetReflector.cpp b/Engine/Source/Developer/SlateReflector/Private/Widgets/SWidgetReflector.cpp
--- a/Engine/Source/Developer/SlateReflector/Private/Widgets/SWidgetReflector.cpp
+++ b/Engine/Source/Developer/SlateReflector/Private/Widgets/SWidgetReflector.cpp
@@ -148,6 +148,14 @@ void SWidgetReflector::Construct( const FArguments& InArgs )
 					+SHorizontalBox::Slot()
 					.AutoWidth()
 					.Padding(5.0f)
+					[
+						SNew(SButton)
+						.OnClicked( this, &SWidgetReflector::CopyStatsToClipboard )
+						.Text( LOCTEXT("CopyStatsToClipboard", "Copy Stats") )
+					]
+					+SHorizontalBox::Slot()
+					.AutoWidth()
+					.Padding(5.0f)
 					[
 						SNew(SCheckBox)
 						.Style( FCoreStyle::Get(), "ToggleButtonCheckbox" )
@@ -278,6 +286,8 @@ void SWidgetReflector::SetWidgetsToVisualize( const FWidgetPath& InWidgetsToVisu
 		FReflectorNode::FindWidgetPath( ReflectorTreeRoot, InWidgetsToVisualize, PickedPath );
 		VisualizeAsTree(PickedPath);
 	}
+	
+	ReflectorTree->RequestTreeRefresh();
 }
 
 
@@ -716,6 +726,22 @@ TSharedRef<ITableRow> SWidgetReflector::GenerateStatRow( TSharedRef<FStatItem> S
 	return SNew(SStatTableRow, OwnerTable, StatItem);
 }
 
+FReply SWidgetReflector::CopyStatsToClipboard()
+{
+	FString ClipboardString;
+
+	for (const TSharedRef<FStatItem>& StatItem : StatsItems)
+	{
+		ClipboardString += StatItem->GetDisplayName().ToString();
+		ClipboardString += TEXT(", ");
+		ClipboardString += StatItem->GetInclusiveAvgMsText().ToString();
+		ClipboardString += TEXT("\r\n");
+	}
+
+	FPlatformMisc::ClipboardCopy( *ClipboardString );
+	return FReply::Handled();
+}
+
 #undef LOCTEXT_NAMESPACE
 
 
diff --git a/Engine/Source/Developer/SlateReflector/Private/Widgets/SWidgetReflector.h b/Engine/Source/Developer/SlateReflector/Private/Widgets/SWidgetReflector.h
--- a/Engine/Source/Developer/SlateReflector/Private/Widgets/SWidgetReflector.h
+++ b/Engine/Source/Developer/SlateReflector/Private/Widgets/SWidgetReflector.h
@@ -208,4 +208,5 @@ private:
 	TSharedRef<SWidget> MakeStatViewer();
 	void UpdateStats();
 	TSharedRef<ITableRow> GenerateStatRow( TSharedRef<FStatItem> StatItem, const TSharedRef<STableViewBase>& OwnerTable );
+	FReply CopyStatsToClipboard();
 };
diff --git a/Engine/Source/Programs/SlateViewer/Private/SlateViewerApp.cpp b/Engine/Source/Programs/SlateViewer/Private/SlateViewerApp.cpp
--- a/Engine/Source/Programs/SlateViewer/Private/SlateViewerApp.cpp
+++ b/Engine/Source/Programs/SlateViewer/Private/SlateViewerApp.cpp
@@ -4,6 +4,7 @@
 #include "RequiredProgramMainCPPInclude.h"
 #include "STestSuite.h"
 #include "ISourceCodeAccessModule.h"
+#include "SPerfSuite.h"
 
 IMPLEMENT_APPLICATION(SlateViewer, "SlateViewer");
 
@@ -38,8 +39,18 @@ void RunSlateViewer( const TCHAR* CommandLine )
 	FGlobalTabmanager::Get()->SetApplicationTitle(NSLOCTEXT("SlateViewer", "AppTitle", "Slate Viewer"));
 	FModuleManager::LoadModuleChecked<ISlateReflectorModule>("SlateReflector").RegisterTabSpawner(WorkspaceMenu::DeveloperMenu);
 
-	// Bring up the test suite.
-	RestoreSlateTestSuite();
+	
+	if (FParse::Param(FCommandLine::Get(), TEXT("perftest")))
+	{
+		// Bring up perf test
+		SummonPerfTestSuite();
+	}
+	else
+	{
+		// Bring up the test suite.
+		RestoreSlateTestSuite();
+	}
+
 
 #if WITH_SHARED_POINTER_TESTS
 	SharedPointerTesting::TestSharedPointer<ESPMode::Fast>();
diff --git a/Engine/Source/Runtime/AppFramework/AppFramework.Build.cs b/Engine/Source/Runtime/AppFramework/AppFramework.Build.cs
--- a/Engine/Source/Runtime/AppFramework/AppFramework.Build.cs
+++ b/Engine/Source/Runtime/AppFramework/AppFramework.Build.cs
@@ -18,6 +18,7 @@ public class AppFramework : ModuleRules
 			new string[] {
 				"InputCore",
 				"WebBrowser",
+				"SlateReflector",
 			}
 		);
 
diff --git a/Engine/Source/Runtime/AppFramework/Private/Framework/Testing/STableViewTesting.cpp b/Engine/Source/Runtime/AppFramework/Private/Framework/Testing/STableViewTesting.cpp
--- a/Engine/Source/Runtime/AppFramework/Private/Framework/Testing/STableViewTesting.cpp
+++ b/Engine/Source/Runtime/AppFramework/Private/Framework/Testing/STableViewTesting.cpp
@@ -792,7 +792,7 @@ public:
 	END_SLATE_FUNCTION_BUILD_OPTIMIZATION
 
 	STableViewTesting()
-	: TotalItems( 125 )
+	: TotalItems( 5000 )
 	, ScrollToIndex( 0 )
 	{
 	}
diff --git a/Engine/Source/Runtime/AppFramework/Private/Widgets/Testing/SPerfSuite.cpp b/Engine/Source/Runtime/AppFramework/Private/Widgets/Testing/SPerfSuite.cpp
new file mode 100644
--- /dev/null
+++ b/Engine/Source/Runtime/AppFramework/Private/Widgets/Testing/SPerfSuite.cpp
@@ -0,0 +1,29 @@
+// Copyright 1998-2014 Epic Games, Inc. All Rights Reserved.
+
+#include "AppFrameworkPrivatePCH.h"
+#include "SPerfSuite.h"
+#include "STableViewTesting.h"
+#include "ISlateReflectorModule.h"
+
+void SummonPerfTestSuite()
+{
+	FSlateApplication::Get().AddWindow
+	( 
+		SNew(SWindow)
+		.IsInitiallyMaximized(false)
+		.ScreenPosition( FVector2D(0,0) )
+		.ClientSize( FVector2D(1920,1200) )
+		[
+			MakeTableViewTesting()
+		]
+	);
+
+	FSlateApplication::Get().AddWindow
+	(
+		SNew(SWindow)
+		.ClientSize(FVector2D(640,800))
+		[
+			FModuleManager::LoadModuleChecked<ISlateReflectorModule>("SlateReflector").GetWidgetReflector()
+		]
+	);
+}
diff --git a/Engine/Source/Runtime/AppFramework/Public/Widgets/Testing/SPerfSuite.h b/Engine/Source/Runtime/AppFramework/Public/Widgets/Testing/SPerfSuite.h
new file mode 100644
--- /dev/null
+++ b/Engine/Source/Runtime/AppFramework/Public/Widgets/Testing/SPerfSuite.h
@@ -0,0 +1,5 @@
+// Copyright 1998-2014 Epic Games, Inc. All Rights Reserved.
+
+#pragma once
+
+void SummonPerfTestSuite();
\ No newline at end of file
