Atlas visualizer improvements, and added a font atlas visualizer

Moved the existing atlas visualizer out of the Slate RHI renderer and instead made the resource manager for each renderer implement the ISlateAtlasProvider interface to provide information about its current texture atlases. FSlateFontCache implements ISlateAtlasProvider for the font atlas viewer.

Improved the atlas visualizer and added support for zooming and panning of the atlas pages. This helps greatly when viewing font atlases as they can be as large as 2048x2048.

diff --git a/Engine/Source/Developer/SlateReflector/Private/SlateReflectorModule.cpp b/Engine/Source/Developer/SlateReflector/Private/SlateReflectorModule.cpp
--- a/Engine/Source/Developer/SlateReflector/Private/SlateReflectorModule.cpp
+++ b/Engine/Source/Developer/SlateReflector/Private/SlateReflectorModule.cpp
@@ -4,6 +4,7 @@
 #include "ISlateReflectorModule.h"
 #include "SDockTab.h"
 #include "ModuleManager.h"
+#include "SAtlasVisualizer.h"
 
 
 #define LOCTEXT_NAMESPACE "FSlateReflectorModule"
@@ -32,25 +33,158 @@ public:
 		return WidgetReflector.ToSharedRef();
 	}
 
+	virtual TSharedRef<SWidget> GetAtlasVisualizer( ISlateAtlasProvider* InAtlasProvider ) override
+	{
+		return SNew(SAtlasVisualizer)
+			.AtlasProvider(InAtlasProvider);
+	}
+
+	virtual TSharedRef<SWidget> GetTextureAtlasVisualizer() override
+	{
+		ISlateAtlasProvider* const AtlasProvider = FSlateApplication::Get().GetRenderer()->GetTextureAtlasProvider();
+
+		if( AtlasProvider )
+		{
+			return GetAtlasVisualizer( AtlasProvider );
+		}
+		else
+		{
+			return SNew(SBox)
+				.HAlign(HAlign_Center)
+				.VAlign(VAlign_Center)
+				[
+					SNew(STextBlock)
+					.Text(LOCTEXT("NoTextureAtlasProvider", "There is no texture atlas provider available for the current renderer."))
+				];
+		}
+	}
+
+	virtual TSharedRef<SWidget> GetFontAtlasVisualizer() override
+	{
+		ISlateAtlasProvider* const AtlasProvider = FSlateApplication::Get().GetRenderer()->GetFontAtlasProvider();
+
+		if( AtlasProvider )
+		{
+			return GetAtlasVisualizer( AtlasProvider );
+		}
+		else
+		{
+			return SNew(SBox)
+				.HAlign(HAlign_Center)
+				.VAlign(VAlign_Center)
+				[
+					SNew(STextBlock)
+					.Text(LOCTEXT("NoFontAtlasProvider", "There is no font atlas provider available for the current renderer."))
+				];
+		}
+	}
+
+	virtual void DisplayWidgetReflector() override
+	{
+		if( bHasRegisteredTabSpawners )
+		{
+			FGlobalTabmanager::Get()->InvokeTab(FTabId("WidgetReflector"));
+		}
+		else if( !WidgetReflectorWindowPtr.IsValid() )
+		{
+			const TSharedRef<SWindow> ReflectorWindow = SNew(SWindow)
+				.AutoCenter(EAutoCenter::PrimaryWorkArea)
+				.ClientSize(FVector2D(600,400))
+				.Title(LOCTEXT("WidgetReflectorTitle", "Widget Reflector"))
+				[
+					GetWidgetReflector()
+				];
+		
+			WidgetReflectorWindowPtr = ReflectorWindow;
+		
+			FSlateApplication::Get().AddWindow( ReflectorWindow );
+		}
+	}
+
+	virtual void DisplayTextureAtlasVisualizer() override
+	{
+		if( bHasRegisteredTabSpawners )
+		{
+			FGlobalTabmanager::Get()->InvokeTab(FTabId("TextureAtlasVisualizer"));
+		}
+		else if( !TextureAtlasVisualizerWindowPtr.IsValid() )
+		{
+			const TSharedRef<SWindow> AtlasVisWindow = SNew(SWindow)
+				.AutoCenter(EAutoCenter::PrimaryWorkArea)
+				.ClientSize(FVector2D(1040, 900))
+				.Title(LOCTEXT("TextureAtlasVisualizerTitle", "Texture Atlas Visualizer"))
+				[
+					GetTextureAtlasVisualizer()
+				];
+		
+			TextureAtlasVisualizerWindowPtr = AtlasVisWindow;
+		
+			FSlateApplication::Get().AddWindow( AtlasVisWindow );
+		}
+	}
+
+	virtual void DisplayFontAtlasVisualizer() override
+	{
+		if( bHasRegisteredTabSpawners )
+		{
+			FGlobalTabmanager::Get()->InvokeTab(FTabId("FontAtlasVisualizer"));
+		}
+		else if( !FontAtlasVisualizerWindowPtr.IsValid() )
+		{
+			const TSharedRef<SWindow> AtlasVisWindow = SNew(SWindow)
+				.AutoCenter(EAutoCenter::PrimaryWorkArea)
+				.ClientSize(FVector2D(1040, 900))
+				.Title(LOCTEXT("FontAtlasVisualizerTitle", "Font Atlas Visualizer"))
+				[
+					GetFontAtlasVisualizer()
+				];
+		
+			FontAtlasVisualizerWindowPtr = AtlasVisWindow;
+		
+			FSlateApplication::Get().AddWindow( AtlasVisWindow );
+		}
+	}
+
 	virtual void RegisterTabSpawner( const TSharedRef<FWorkspaceItem>& WorkspaceGroup ) override
 	{
+		bHasRegisteredTabSpawners = true;
+
 		FGlobalTabmanager::Get()->RegisterNomadTabSpawner("WidgetReflector", FOnSpawnTab::CreateRaw(this, &FSlateReflectorModule::MakeWidgetReflectorTab) )
 			.SetDisplayName(LOCTEXT("WidgetReflectorTitle", "Widget Reflector"))
 			.SetTooltipText(LOCTEXT("WidgetReflectorTooltipText", "Open the Widget Reflector tab."))
 			.SetGroup(WorkspaceGroup)
 			.SetIcon(FSlateIcon(FCoreStyle::Get().GetStyleSetName(), "WidgetReflector.TabIcon"));
+
+		FGlobalTabmanager::Get()->RegisterNomadTabSpawner("TextureAtlasVisualizer", FOnSpawnTab::CreateRaw(this, &FSlateReflectorModule::MakeTextureAtlasVisualizerTab) )
+			.SetDisplayName(LOCTEXT("TextureAtlasVisualizerTitle", "Texture Atlas Visualizer"))
+			.SetTooltipText(LOCTEXT("TextureAtlasVisualizerTooltipText", "Open the Texture Atlas Visualizer tab."))
+			.SetGroup(WorkspaceGroup)
+			.SetMenuType(ETabSpawnerMenuType::Hidden);
+
+		FGlobalTabmanager::Get()->RegisterNomadTabSpawner("FontAtlasVisualizer", FOnSpawnTab::CreateRaw(this, &FSlateReflectorModule::MakeFontAtlasVisualizerTab) )
+			.SetDisplayName(LOCTEXT("FontAtlasVisualizerTitle", "Font Atlas Visualizer"))
+			.SetTooltipText(LOCTEXT("FontAtlasVisualizerTooltipText", "Open the Font Atlas Visualizer tab."))
+			.SetGroup(WorkspaceGroup)
+			.SetMenuType(ETabSpawnerMenuType::Hidden);
 	}
 
 	virtual void UnregisterTabSpawner() 
 	{
+		bHasRegisteredTabSpawners = false;
+
 		FGlobalTabmanager::Get()->UnregisterNomadTabSpawner("WidgetReflector");
+		FGlobalTabmanager::Get()->UnregisterNomadTabSpawner("TextureAtlasVisualizer");
+		FGlobalTabmanager::Get()->UnregisterNomadTabSpawner("FontAtlasVisualizer");
 	}
 
 public:
 
 	// IModuleInterface interface
 
-	virtual void StartupModule() override { }
+	virtual void StartupModule() override
+	{
+		bHasRegisteredTabSpawners = false;
+	}
 
 	virtual void ShutdownModule() override
 	{
@@ -68,10 +202,40 @@ private:
 			];
 	}
 
+	TSharedRef<SDockTab> MakeTextureAtlasVisualizerTab( const FSpawnTabArgs& )
+	{
+		return SNew(SDockTab)
+			.TabRole(ETabRole::NomadTab)
+			[
+				GetTextureAtlasVisualizer()
+			];
+	}
+
+	TSharedRef<SDockTab> MakeFontAtlasVisualizerTab( const FSpawnTabArgs& )
+	{
+		return SNew(SDockTab)
+			.TabRole(ETabRole::NomadTab)
+			[
+				GetFontAtlasVisualizer()
+			];
+	}
+
 private:
 
+	/** True if the tab spawners have been registered for this module */
+	bool bHasRegisteredTabSpawners;
+
 	/** Holds the widget reflector singleton. */
 	TWeakPtr<SWidgetReflector> WidgetReflectorPtr;
+
+	/** Holds the widget reflector singleton. */
+	TWeakPtr<SWindow> WidgetReflectorWindowPtr;
+
+	/** Holds the texture atlas viewer singleton. */
+	TWeakPtr<SWindow> TextureAtlasVisualizerWindowPtr;
+
+	/** Holds the font atlas viewer singleton. */
+	TWeakPtr<SWindow> FontAtlasVisualizerWindowPtr;
 };
 
 
diff --git a/Engine/Source/Developer/SlateReflector/Private/Widgets/SAtlasVisualizer.cpp b/Engine/Source/Developer/SlateReflector/Private/Widgets/SAtlasVisualizer.cpp
new file mode 100644
--- /dev/null
+++ b/Engine/Source/Developer/SlateReflector/Private/Widgets/SAtlasVisualizer.cpp
@@ -0,0 +1,508 @@
+// Copyright 1998-2014 Epic Games, Inc. All Rights Reserved.
+
+#include "SlateReflectorPrivatePCH.h"
+#include "SAtlasVisualizer.h"
+#include "ScrollyZoomy.h"
+
+#define LOCTEXT_NAMESPACE "AtlasVisualizer"
+
+class SAtlasVisualizerPanel : public IScrollableZoomable, public SPanel
+{
+public:
+	class FAtlasVisualizerPanelSlot : public TSupportsOneChildMixin<FAtlasVisualizerPanelSlot>
+	{
+	public:
+		FAtlasVisualizerPanelSlot()
+			: TSupportsOneChildMixin<FAtlasVisualizerPanelSlot>()
+		{
+		}
+	};
+
+	SLATE_BEGIN_ARGS(SAtlasVisualizerPanel)
+		{
+			_Visibility = EVisibility::Visible;
+		}
+		
+		SLATE_DEFAULT_SLOT( FArguments, Content )
+
+	SLATE_END_ARGS()
+
+	SAtlasVisualizerPanel()
+		: PhysicalOffset(ForceInitToZero)
+		, CachedSize(ForceInitToZero)
+		, ZoomLevel(1.0f)
+		, bFitToWindow(true)
+		, ChildSlot()
+		, ScrollyZoomy(false)
+	{
+	}
+
+	void Construct( const FArguments& InArgs )
+	{
+		ChildSlot
+		[
+			InArgs._Content.Widget
+		];
+	}
+
+	virtual void OnArrangeChildren( const FGeometry& AllottedGeometry, FArrangedChildren& ArrangedChildren ) const override
+	{
+		CachedSize = AllottedGeometry.GetLocalSize();
+
+		const TSharedRef<SWidget>& ChildWidget = ChildSlot.GetWidget();
+		if( ChildWidget->GetVisibility() != EVisibility::Collapsed )
+		{
+			const FVector2D& WidgetDesiredSize = ChildWidget->GetDesiredSize();
+
+			// Update the zoom level, and clamp the pan offset based on our current geometry
+			SAtlasVisualizerPanel* const NonConstThis = const_cast<SAtlasVisualizerPanel*>(this);
+			NonConstThis->UpdateFitToWindowZoom(WidgetDesiredSize, CachedSize);
+			NonConstThis->ClampViewOffset(WidgetDesiredSize, CachedSize);
+
+			ArrangedChildren.AddWidget(AllottedGeometry.MakeChild(ChildWidget, PhysicalOffset * ZoomLevel, WidgetDesiredSize * ZoomLevel));
+		}
+	}
+
+	virtual FVector2D ComputeDesiredSize() const override
+	{
+		FVector2D ThisDesiredSize = FVector2D::ZeroVector;
+
+		const TSharedRef<SWidget>& ChildWidget = ChildSlot.GetWidget();
+		if( ChildWidget->GetVisibility() != EVisibility::Collapsed )
+		{
+			ThisDesiredSize = ChildWidget->GetDesiredSize() * ZoomLevel;
+		}
+
+		return ThisDesiredSize;
+	}
+
+	virtual FChildren* GetChildren() override
+	{
+		return &ChildSlot;
+	}
+
+	virtual void Tick( const FGeometry& AllottedGeometry, const double InCurrentTime, const float InDeltaTime ) override
+	{
+		ScrollyZoomy.Tick(InDeltaTime, *this);
+	}
+
+	virtual FReply OnMouseButtonDown( const FGeometry& MyGeometry, const FPointerEvent& MouseEvent ) override
+	{
+		return ScrollyZoomy.OnMouseButtonDown(MouseEvent);
+	}
+
+	virtual FReply OnMouseButtonUp( const FGeometry& MyGeometry, const FPointerEvent& MouseEvent ) override
+	{
+		return ScrollyZoomy.OnMouseButtonUp(AsShared(), MyGeometry, MouseEvent);
+	}
+		
+	virtual FReply OnMouseMove( const FGeometry& MyGeometry, const FPointerEvent& MouseEvent ) override
+	{
+		return ScrollyZoomy.OnMouseMove(AsShared(), *this, MyGeometry, MouseEvent);
+	}
+
+	virtual void OnMouseLeave( const FPointerEvent& MouseEvent ) override
+	{
+		ScrollyZoomy.OnMouseLeave(AsShared(), MouseEvent);
+	}
+		
+	virtual FReply OnMouseWheel( const FGeometry& MyGeometry, const FPointerEvent& MouseEvent ) override
+	{
+		return ScrollyZoomy.OnMouseWheel(MouseEvent, *this);
+	}
+
+	virtual FCursorReply OnCursorQuery( const FGeometry& MyGeometry, const FPointerEvent& CursorEvent ) const override
+	{
+		return ScrollyZoomy.OnCursorQuery();
+	}
+
+	virtual int32 OnPaint(const FPaintArgs& Args, const FGeometry& AllottedGeometry, const FSlateRect& MyClippingRect, FSlateWindowElementList& OutDrawElements, int32 LayerId, const FWidgetStyle& InWidgetStyle, bool bParentEnabled) const override
+	{
+		LayerId = SPanel::OnPaint(Args, AllottedGeometry, MyClippingRect, OutDrawElements, LayerId, InWidgetStyle, bParentEnabled);
+		LayerId = ScrollyZoomy.PaintSoftwareCursorIfNeeded(AllottedGeometry, MyClippingRect, OutDrawElements, LayerId);
+		return LayerId;
+	}
+
+	virtual bool ScrollBy( const FVector2D& Offset ) override
+	{
+		if( bFitToWindow )
+		{
+			return false;
+		}
+
+		const FVector2D PrevPhysicalOffset = PhysicalOffset;
+		const float InverseZoom = 1.0f / ZoomLevel;
+		PhysicalOffset += (Offset * InverseZoom);
+
+		const TSharedRef<SWidget>& ChildWidget = ChildSlot.GetWidget();
+		const FVector2D& WidgetDesiredSize = ChildWidget->GetDesiredSize();
+		ClampViewOffset(WidgetDesiredSize, CachedSize);
+
+		return PhysicalOffset != PrevPhysicalOffset;
+	}
+
+	virtual bool ZoomBy( const float Amount ) override
+	{
+		static const float MinZoomLevel = 0.2f;
+		static const float MaxZoomLevel = 4.0f;
+
+		bFitToWindow = false;
+
+		const float PrevZoomLevel = ZoomLevel;
+		ZoomLevel = FMath::Clamp(ZoomLevel + (Amount * 0.05f), MinZoomLevel, MaxZoomLevel);
+		return ZoomLevel != PrevZoomLevel;
+	}
+
+	float GetZoomLevel() const
+	{
+		return ZoomLevel;
+	}
+
+	void FitToWindow()
+	{
+		bFitToWindow = true;
+		PhysicalOffset = FVector2D::ZeroVector;
+	}
+
+	bool IsFitToWindow() const
+	{
+		return bFitToWindow;
+	}
+
+	void FitToSize()
+	{
+		bFitToWindow = false;
+		ZoomLevel = 1.0f;
+		PhysicalOffset = FVector2D::ZeroVector;
+	}
+
+private:
+	void UpdateFitToWindowZoom( const FVector2D& ViewportSize, const FVector2D& LocalSize )
+	{
+		if( bFitToWindow )
+		{
+			const float ZoomHoriz = LocalSize.X / ViewportSize.X;
+			const float ZoomVert  = LocalSize.Y / ViewportSize.Y;
+			ZoomLevel = FMath::Min(ZoomHoriz, ZoomVert);
+		}
+	}
+
+	void ClampViewOffset( const FVector2D& ViewportSize, const FVector2D& LocalSize )
+	{
+		PhysicalOffset.X = ClampViewOffsetAxis(ViewportSize.X, LocalSize.X, PhysicalOffset.X);
+		PhysicalOffset.Y = ClampViewOffsetAxis(ViewportSize.Y, LocalSize.Y, PhysicalOffset.Y);
+	}
+
+	float ClampViewOffsetAxis( const float ViewportSize, const float LocalSize, const float CurrentOffset )
+	{
+		const float ZoomedViewportSize = ViewportSize * ZoomLevel;
+
+		if( ZoomedViewportSize <= LocalSize )
+		{
+			// If the viewport is smaller than the available size, then we can't be scrolled
+			return 0.0f;
+		}
+
+		// Given the size of the viewport, and the current size of the window, work how far we can scroll
+		// Note: This number is negative since scrolling down/right moves the viewport up/left
+		const float MaxScrollOffset = (LocalSize - ZoomedViewportSize) / ZoomLevel;
+
+		// Clamp the left/top edge
+		if( CurrentOffset < MaxScrollOffset )
+		{
+			return MaxScrollOffset;
+		}
+
+		// Clamp the right/bottom edge
+		if( CurrentOffset > 0.0f )
+		{
+			return 0.0f;
+		}
+
+		return CurrentOffset;
+	}
+
+	FVector2D PhysicalOffset;
+	mutable FVector2D CachedSize;
+	float ZoomLevel;
+	bool bFitToWindow;
+
+	FAtlasVisualizerPanelSlot ChildSlot;
+	FScrollyZoomy ScrollyZoomy;
+};
+
+void SAtlasVisualizer::Construct( const FArguments& InArgs )
+{
+	AtlasProvider = InArgs._AtlasProvider;
+	check(AtlasProvider);
+
+	const bool IsAlphaOnly = AtlasProvider->IsAtlasPageResourceAlphaOnly();
+
+	SelectedAtlasPage = 0;
+	bDisplayCheckerboard = false;
+
+	const FIntPoint DesiredViewportSize = GetSize();
+	TSharedPtr<SViewport> Viewport;
+
+	ChildSlot
+	[
+		SNew(SBorder)
+		.BorderImage(FCoreStyle::Get().GetBrush("ToolPanel.GroupBorder"))
+		[
+			SNew( SVerticalBox )
+			+ SVerticalBox::Slot()
+			.Padding(4.0f)
+			.AutoHeight()
+			[
+				SNew( SHorizontalBox )
+				+ SHorizontalBox::Slot()
+				.VAlign( VAlign_Center )
+				.AutoWidth()
+				.Padding(0.0,2.0f,2.0f,2.0f)
+				[
+					SNew( STextBlock )
+					.Text( LOCTEXT("SelectAPage", "Select a page") )
+				]
+				+ SHorizontalBox::Slot()
+				.AutoWidth()
+				.Padding(2.0f)
+				.VAlign( VAlign_Center )
+				[
+					SAssignNew( AtlasPageCombo, SComboBox< TSharedPtr<int32> > )
+					.OptionsSource( &AtlasPages )
+					.OnComboBoxOpening(this, &SAtlasVisualizer::OnComboOpening)
+					.OnGenerateWidget( this, &SAtlasVisualizer::OnGenerateWidgetForCombo )
+					.OnSelectionChanged( this, &SAtlasVisualizer::OnAtlasPageChanged )
+					.Content()
+					[
+						SNew( STextBlock )
+						.Text( this, &SAtlasVisualizer::OnGetSelectedItemText )
+					]
+				]
+				+ SHorizontalBox::Slot()
+				.VAlign( VAlign_Center )
+				.AutoWidth()
+				.Padding(0.0,2.0f,2.0f,2.0f)
+				[
+					SNew( STextBlock )
+					.Text( FText::Format( LOCTEXT("PageSizeXY", "({0} x {1})"), FText::AsNumber(DesiredViewportSize.X), FText::AsNumber(DesiredViewportSize.Y) ) )
+				]
+				+ SHorizontalBox::Slot()
+				.Padding(20.0f,2.0f)
+				.AutoWidth()
+				.VAlign( VAlign_Center )
+				[
+					SNew( SCheckBox )
+					.Visibility( (IsAlphaOnly) ? EVisibility::Collapsed : EVisibility::Visible )
+					.OnCheckStateChanged( this, &SAtlasVisualizer::OnDisplayCheckerboardStateChanged )
+					.IsChecked( this, &SAtlasVisualizer::OnGetCheckerboardState )
+					.Content()
+					[
+						SNew( STextBlock )
+						.Text( LOCTEXT("DisplayCheckerboardCheckboxLabel", "Display Checkerboard") )
+					]
+				]
+
+				+ SHorizontalBox::Slot()
+				[
+					SNew(SSpacer)
+				]
+
+				+ SHorizontalBox::Slot()
+				.Padding(2.0f)
+				.AutoWidth()
+				.VAlign( VAlign_Center )
+				[
+					SNew( STextBlock )
+					.Text( this, &SAtlasVisualizer::GetZoomLevelPercentText )
+				]
+
+				+ SHorizontalBox::Slot()
+				.Padding(2.0f)
+				.AutoWidth()
+				.VAlign( VAlign_Center )
+				[
+					SNew( SCheckBox )
+					.OnCheckStateChanged( this, &SAtlasVisualizer::OnFitToWindowStateChanged )
+					.IsChecked( this, &SAtlasVisualizer::OnGetFitToWindowState )
+					.Content()
+					[
+						SNew( STextBlock )
+						.Text( LOCTEXT("FitToWindow", "Fit to Window") )
+					]
+				]
+
+				+ SHorizontalBox::Slot()
+				.Padding(2.0f)
+				.AutoWidth()
+				.VAlign( VAlign_Center )
+				[
+					SNew( SButton )
+					.Text( LOCTEXT("ActualSize", "Actual Size") )
+					.OnClicked( this, &SAtlasVisualizer::OnActualSizeClicked )
+				]
+			]
+			+ SVerticalBox::Slot()
+			.Padding(2.0f)
+			[
+				SAssignNew( ScrollPanel, SAtlasVisualizerPanel )
+				[
+					SNew( SOverlay )
+					+ SOverlay::Slot()
+					[
+						SNew( SImage )
+						.Visibility( this, &SAtlasVisualizer::OnGetCheckerboardVisibility )
+						.Image( FCoreStyle::Get().GetBrush("Checkerboard") )
+					]
+					+ SOverlay::Slot()
+					[
+						SAssignNew( Viewport, SViewport )
+						.ViewportSize(FVector2D(DesiredViewportSize.X, DesiredViewportSize.Y))
+						.IgnoreTextureAlpha(false)
+						.EnableBlending(true)
+					]
+				]
+			]
+		]
+	];
+
+	Viewport->SetViewportInterface( SharedThis( this ) );
+}
+
+FIntPoint SAtlasVisualizer::GetSize() const
+{
+	return AtlasProvider->GetAtlasPageSize();
+}
+
+bool SAtlasVisualizer::RequiresVsync() const
+{
+	return false;
+}
+
+FSlateShaderResource* SAtlasVisualizer::GetViewportRenderTargetTexture() const
+{
+	if( SelectedAtlasPage < AtlasProvider->GetNumAtlasPages() )
+	{
+		return AtlasProvider->GetAtlasPageResource(SelectedAtlasPage);
+	}
+
+	return nullptr;
+}
+
+bool SAtlasVisualizer::IsViewportTextureAlphaOnly() const
+{
+	return AtlasProvider->IsAtlasPageResourceAlphaOnly();
+}
+
+FText SAtlasVisualizer::GetZoomLevelPercentText() const
+{
+	if( ScrollPanel.IsValid() )
+	{
+		return FText::Format( LOCTEXT("ZoomLevelPercent", "Zoom Level: {0}"), FText::AsPercent(ScrollPanel->GetZoomLevel()) );
+	}
+
+	return FText::GetEmpty();
+}
+
+void SAtlasVisualizer::OnFitToWindowStateChanged( ESlateCheckBoxState::Type NewState )
+{
+	if( ScrollPanel.IsValid() )
+	{
+		if( NewState == ESlateCheckBoxState::Checked )
+		{
+			ScrollPanel->FitToWindow();
+		}
+		else
+		{
+			ScrollPanel->FitToSize();
+		}
+	}
+}
+
+ESlateCheckBoxState::Type SAtlasVisualizer::OnGetFitToWindowState() const
+{
+	if( ScrollPanel.IsValid() )
+	{
+		return (ScrollPanel->IsFitToWindow()) ? ESlateCheckBoxState::Checked : ESlateCheckBoxState::Unchecked;
+	}
+
+	return ESlateCheckBoxState::Undetermined;
+}
+
+FReply SAtlasVisualizer::OnActualSizeClicked()
+{
+	if( ScrollPanel.IsValid() )
+	{
+		ScrollPanel->FitToSize();
+	}
+
+	return FReply::Handled();
+}
+
+void SAtlasVisualizer::OnDisplayCheckerboardStateChanged( ESlateCheckBoxState::Type NewState )
+{
+	bDisplayCheckerboard = NewState == ESlateCheckBoxState::Checked;
+}
+
+ESlateCheckBoxState::Type SAtlasVisualizer::OnGetCheckerboardState() const
+{
+	return bDisplayCheckerboard ? ESlateCheckBoxState::Checked : ESlateCheckBoxState::Unchecked;
+}
+
+EVisibility SAtlasVisualizer::OnGetCheckerboardVisibility() const
+{
+	return bDisplayCheckerboard ? EVisibility::Visible : EVisibility::Collapsed;
+}
+
+void SAtlasVisualizer::OnComboOpening()
+{
+	const int32 NumAtlasPages = AtlasProvider->GetNumAtlasPages();
+
+	AtlasPages.Empty();
+	for( int32 AtlasIndex = 0; AtlasIndex < NumAtlasPages; ++AtlasIndex )
+	{
+		AtlasPages.Add( MakeShareable( new int32( AtlasIndex ) ) );
+	}
+
+	TSharedPtr<int32> SelectedComboEntry;
+	if( SelectedAtlasPage < NumAtlasPages )
+	{
+		SelectedComboEntry = AtlasPages[SelectedAtlasPage];
+	}
+	else if( AtlasPages.Num() > 0 )
+	{
+		SelectedAtlasPage = 0;
+		SelectedComboEntry = AtlasPages[0];
+	}
+
+	AtlasPageCombo->ClearSelection();
+	AtlasPageCombo->RefreshOptions();
+	AtlasPageCombo->SetSelectedItem(SelectedComboEntry);
+}
+
+FText SAtlasVisualizer::OnGetSelectedItemText() const
+{
+	if( SelectedAtlasPage < AtlasProvider->GetNumAtlasPages() )
+	{
+		return FText::Format( LOCTEXT("PageX", "Page {0}"), FText::AsNumber(SelectedAtlasPage) );
+	}
+	
+	return LOCTEXT("SelectAPage", "Select a page");
+}
+
+void SAtlasVisualizer::OnAtlasPageChanged( TSharedPtr<int32> AtlasPage, ESelectInfo::Type SelectionType )
+{
+	if( AtlasPage.IsValid() )
+	{
+		SelectedAtlasPage = *AtlasPage;
+	}
+}
+
+TSharedRef<SWidget> SAtlasVisualizer::OnGenerateWidgetForCombo( TSharedPtr<int32> AtlasPage )
+{
+	return 
+	SNew( STextBlock )
+	.Text( FText::Format( LOCTEXT("PageX", "Page {0}"), FText::AsNumber(*AtlasPage) ) );
+}
+
+#undef LOCTEXT_NAMESPACE
diff --git a/Engine/Source/Developer/SlateReflector/Private/Widgets/SAtlasVisualizer.h b/Engine/Source/Developer/SlateReflector/Private/Widgets/SAtlasVisualizer.h
new file mode 100644
--- /dev/null
+++ b/Engine/Source/Developer/SlateReflector/Private/Widgets/SAtlasVisualizer.h
@@ -0,0 +1,46 @@
+// Copyright 1998-2014 Epic Games, Inc. All Rights Reserved.
+
+#pragma once
+
+class ISlateAtlasProvider;
+class SAtlasVisualizerPanel;
+
+class SAtlasVisualizer : public ISlateViewport, public SCompoundWidget
+{
+public:
+	SLATE_BEGIN_ARGS( SAtlasVisualizer )
+		: _AtlasProvider()
+		{}
+
+		SLATE_ARGUMENT( ISlateAtlasProvider*, AtlasProvider )
+
+	SLATE_END_ARGS()
+
+	// ISlateViewport
+	virtual FIntPoint GetSize() const override;
+	virtual bool RequiresVsync() const override;
+	virtual FSlateShaderResource* GetViewportRenderTargetTexture() const override;
+	bool IsViewportTextureAlphaOnly() const override;
+
+	void Construct( const FArguments& InArgs );
+
+protected:
+	FText GetZoomLevelPercentText() const;
+	void OnFitToWindowStateChanged( ESlateCheckBoxState::Type NewState );
+	ESlateCheckBoxState::Type OnGetFitToWindowState() const;
+	FReply OnActualSizeClicked();
+	void OnDisplayCheckerboardStateChanged( ESlateCheckBoxState::Type NewState );
+	ESlateCheckBoxState::Type OnGetCheckerboardState() const;
+	EVisibility OnGetCheckerboardVisibility() const;
+	void OnComboOpening();
+	FText OnGetSelectedItemText() const;
+	void OnAtlasPageChanged( TSharedPtr<int32> AtlasPage, ESelectInfo::Type SelectionType );
+	TSharedRef<SWidget> OnGenerateWidgetForCombo( TSharedPtr<int32> AtlasPage );
+
+	ISlateAtlasProvider* AtlasProvider;
+	TSharedPtr<SComboBox<TSharedPtr<int32>>> AtlasPageCombo;
+	TArray<TSharedPtr<int32>> AtlasPages;
+	TSharedPtr<SAtlasVisualizerPanel> ScrollPanel;
+	int32 SelectedAtlasPage;
+	bool bDisplayCheckerboard;
+};
diff --git a/Engine/Source/Developer/SlateReflector/Private/Widgets/SWidgetReflector.cpp b/Engine/Source/Developer/SlateReflector/Private/Widgets/SWidgetReflector.cpp
--- a/Engine/Source/Developer/SlateReflector/Private/Widgets/SWidgetReflector.cpp
+++ b/Engine/Source/Developer/SlateReflector/Private/Widgets/SWidgetReflector.cpp
@@ -1,6 +1,7 @@
 // Copyright 1998-2014 Epic Games, Inc. All Rights Reserved.
 
 #include "SlateReflectorPrivatePCH.h"
+#include "ISlateReflectorModule.h"
 
 #if STATS
 #include "StatsData.h"
@@ -73,20 +74,26 @@ void SWidgetReflector::Construct( const FArguments& InArgs )
 					
 						+ SHorizontalBox::Slot()
 							.AutoWidth()
+							.VAlign(VAlign_Center)
 							[
 								SNew(STextBlock)
 									.Text(LOCTEXT("AppScale", "Application Scale: ").ToString())
 							]
 
 						+ SHorizontalBox::Slot()
-							.MaxWidth(250)
+							.AutoWidth()
 							[
-								SNew(SSpinBox<float>)
-									.Value(this, &SWidgetReflector::HandleAppScaleSliderValue)
-									.MinValue(0.1f)
-									.MaxValue(3.0f)
-									.Delta(0.01f)
-									.OnValueChanged(this, &SWidgetReflector::HandleAppScaleSliderChanged)
+								SNew(SBox)
+								.MinDesiredWidth(100)
+								.MaxDesiredWidth(250)
+								[
+									SNew(SSpinBox<float>)
+										.Value(this, &SWidgetReflector::HandleAppScaleSliderValue)
+										.MinValue(0.1f)
+										.MaxValue(3.0f)
+										.Delta(0.01f)
+										.OnValueChanged(this, &SWidgetReflector::HandleAppScaleSliderChanged)
+								]
 							]
 						+ SHorizontalBox::Slot()
 						.FillWidth(1.0f)
@@ -94,12 +101,21 @@ void SWidgetReflector::Construct( const FArguments& InArgs )
 							SNew(SSpacer)
 						]
 						+SHorizontalBox::Slot()
-						.HAlign(HAlign_Right)
+						.AutoWidth()
+						.Padding(FMargin(5.0f, 0.0f))
 						[
 							SNew(SButton)
-							.Text(LOCTEXT("DisplayTextureAtlases", "Display Atlases"))
+							.Text(LOCTEXT("DisplayTextureAtlases", "Display Texture Atlases"))
 							.OnClicked(this, &SWidgetReflector::HandleDisplayTextureAtlases)
 						]
+						+SHorizontalBox::Slot()
+						.AutoWidth()
+						.Padding(FMargin(5.0f, 0.0f))
+						[
+							SNew(SButton)
+							.Text(LOCTEXT("DisplayFontAtlases", "Display Font Atlases"))
+							.OnClicked(this, &SWidgetReflector::HandleDisplayFontAtlases)
+						]
 					]
 
 				+ SVerticalBox::Slot()
@@ -428,6 +444,20 @@ int32 SWidgetReflector::VisualizeSelectedNodesAsRectangles( const TArray<TShared
 /* SWidgetReflector callbacks
  *****************************************************************************/
 
+FReply SWidgetReflector::HandleDisplayTextureAtlases()
+{
+	static const FName SlateReflectorModuleName("SlateReflector");
+	FModuleManager::LoadModuleChecked<ISlateReflectorModule>(SlateReflectorModuleName).DisplayTextureAtlasVisualizer();
+	return FReply::Handled();
+}
+
+FReply SWidgetReflector::HandleDisplayFontAtlases()
+{
+	static const FName SlateReflectorModuleName("SlateReflector");
+	FModuleManager::LoadModuleChecked<ISlateReflectorModule>(SlateReflectorModuleName).DisplayFontAtlasVisualizer();
+	return FReply::Handled();
+}
+
 void SWidgetReflector::HandleFocusCheckBoxCheckedStateChanged( ESlateCheckBoxState::Type NewValue )
 {
 	bShowFocus = NewValue != ESlateCheckBoxState::Unchecked;
diff --git a/Engine/Source/Developer/SlateReflector/Private/Widgets/SWidgetReflector.h b/Engine/Source/Developer/SlateReflector/Private/Widgets/SWidgetReflector.h
--- a/Engine/Source/Developer/SlateReflector/Private/Widgets/SWidgetReflector.h
+++ b/Engine/Source/Developer/SlateReflector/Private/Widgets/SWidgetReflector.h
@@ -124,11 +124,9 @@ private:
 		FSlateApplication::Get().SetApplicationScale(NewValue);
 	}
 
-	FReply HandleDisplayTextureAtlases()
-	{
-		FSlateApplication::Get().GetRenderer()->DisplayTextureAtlases();
-		return FReply::Handled();
-	}
+	FReply HandleDisplayTextureAtlases();
+
+	FReply HandleDisplayFontAtlases();
 
 	/** Callback for getting the value of the application scale slider. */
 	float HandleAppScaleSliderValue() const
diff --git a/Engine/Source/Developer/SlateReflector/Public/ISlateReflectorModule.h b/Engine/Source/Developer/SlateReflector/Public/ISlateReflectorModule.h
--- a/Engine/Source/Developer/SlateReflector/Public/ISlateReflectorModule.h
+++ b/Engine/Source/Developer/SlateReflector/Public/ISlateReflectorModule.h
@@ -8,6 +8,7 @@
 // forward declarations
 class FWorkspaceItem;
 class SWidget;
+class ISlateAtlasProvider;
 
 
 /**
@@ -25,6 +26,42 @@ public:
 	 */
 	virtual TSharedRef<SWidget> GetWidgetReflector() = 0;
 
+	/**
+	 * Create an atlas visualizer widget using the given atlas provider.
+	 *
+	 * @return The widget.
+	 */
+	virtual TSharedRef<SWidget> GetAtlasVisualizer( ISlateAtlasProvider* InAtlasProvider ) = 0;
+
+	/**
+	 * Create an atlas visualizer widget for the texture atlas provider used by the current renderer.
+	 *
+	 * @return The widget.
+	 */
+	virtual TSharedRef<SWidget> GetTextureAtlasVisualizer() = 0;
+
+	/**
+	 * Create an atlas visualizer widget for the font atlas provider used by the current renderer.
+	 *
+	 * @return The widget.
+	 */
+	virtual TSharedRef<SWidget> GetFontAtlasVisualizer() = 0;
+
+	/**
+	 * Display the widget reflector, either spawned from a tab manager, or in a new window if the tab manager can't be used
+	 */
+	virtual void DisplayWidgetReflector() = 0;
+
+	/**
+	 * Display the texture atlas visualizer, either spawned from a tab manager, or in a new window if the tab manager can't be used
+	 */
+	virtual void DisplayTextureAtlasVisualizer() = 0;
+
+	/**
+	 * Display the texture atlas visualizer, either spawned from a tab manager, or in a new window if the tab manager can't be used
+	 */
+	virtual void DisplayFontAtlasVisualizer() = 0;
+
 	/**
 	 * Registers a tab spawner for the widget reflector.
 	 *
diff --git a/Engine/Source/Developer/StandaloneRenderer/Private/OpenGL/SlateOpenGLRenderer.cpp b/Engine/Source/Developer/StandaloneRenderer/Private/OpenGL/SlateOpenGLRenderer.cpp
--- a/Engine/Source/Developer/StandaloneRenderer/Private/OpenGL/SlateOpenGLRenderer.cpp
+++ b/Engine/Source/Developer/StandaloneRenderer/Private/OpenGL/SlateOpenGLRenderer.cpp
@@ -38,16 +38,23 @@ public:
 	{
 	}
 
-	virtual TSharedRef<FSlateFontAtlas> CreateFontAtlas() const override
+	virtual FIntPoint GetAtlasSize() const override
 	{
-		/** Size of each font texture, width and height */
-		const uint32 TextureSize = 1024;
+		return FIntPoint(TextureSize, TextureSize);
+	}
 
+	virtual TSharedRef<FSlateFontAtlas> CreateFontAtlas() const override
+	{
 		TSharedRef<FSlateFontTextureOpenGL> FontTexture = MakeShareable( new FSlateFontTextureOpenGL( TextureSize, TextureSize ) );
 		FontTexture->CreateFontTexture();
 
 		return FontTexture;
 	}
+
+private:
+
+	/** Size of each font texture, width and height */
+	static const uint32 TextureSize = 1024;
 };
 
 void FSlateOpenGLRenderer::Initialize()
@@ -234,3 +241,13 @@ void FSlateOpenGLRenderer::ReleaseUpdatableTexture(FSlateUpdatableTexture* Textu
 {
 	delete Texture;
 }
+
+ISlateAtlasProvider* FSlateOpenGLRenderer::GetTextureAtlasProvider()
+{
+	if( TextureManager.IsValid() )
+	{
+		return TextureManager->GetTextureAtlasProvider();
+	}
+
+	return nullptr;
+}
diff --git a/Engine/Source/Developer/StandaloneRenderer/Private/OpenGL/SlateOpenGLRenderer.h b/Engine/Source/Developer/StandaloneRenderer/Private/OpenGL/SlateOpenGLRenderer.h
--- a/Engine/Source/Developer/StandaloneRenderer/Private/OpenGL/SlateOpenGLRenderer.h
+++ b/Engine/Source/Developer/StandaloneRenderer/Private/OpenGL/SlateOpenGLRenderer.h
@@ -148,6 +148,7 @@ public:
 	virtual void LoadStyleResources(const ISlateStyle& Style) override;
 	virtual FSlateUpdatableTexture* CreateUpdatableTexture(uint32 Width, uint32 Height) override;
 	virtual void ReleaseUpdatableTexture(FSlateUpdatableTexture* Texture) override;
+	virtual ISlateAtlasProvider* GetTextureAtlasProvider() override;
 
 private:
 	/** 
diff --git a/Engine/Source/Developer/StandaloneRenderer/Private/OpenGL/SlateOpenGLTextureManager.cpp b/Engine/Source/Developer/StandaloneRenderer/Private/OpenGL/SlateOpenGLTextureManager.cpp
--- a/Engine/Source/Developer/StandaloneRenderer/Private/OpenGL/SlateOpenGLTextureManager.cpp
+++ b/Engine/Source/Developer/StandaloneRenderer/Private/OpenGL/SlateOpenGLTextureManager.cpp
@@ -312,3 +312,9 @@ FSlateShaderResourceProxy* FSlateOpenGLTextureManager::GetShaderResource( const
 
 	return Texture;
 }
+
+ISlateAtlasProvider* FSlateOpenGLTextureManager::GetTextureAtlasProvider()
+{
+	// Texture atlases aren't implemented for the standalone OpenGL renderer
+	return nullptr;
+}
diff --git a/Engine/Source/Developer/StandaloneRenderer/Private/OpenGL/SlateOpenGLTextureManager.h b/Engine/Source/Developer/StandaloneRenderer/Private/OpenGL/SlateOpenGLTextureManager.h
--- a/Engine/Source/Developer/StandaloneRenderer/Private/OpenGL/SlateOpenGLTextureManager.h
+++ b/Engine/Source/Developer/StandaloneRenderer/Private/OpenGL/SlateOpenGLTextureManager.h
@@ -22,6 +22,8 @@ public:
 	 */
 	virtual FSlateShaderResourceProxy* GetShaderResource( const FSlateBrush& InBrush ) override;
 
+	virtual ISlateAtlasProvider* GetTextureAtlasProvider() override;
+
 	/**
 	 * Creates a 1x1 texture of the specified color
 	 * 
diff --git a/Engine/Source/Developer/StandaloneRenderer/Private/Windows/D3D/SlateD3DRenderer.cpp b/Engine/Source/Developer/StandaloneRenderer/Private/Windows/D3D/SlateD3DRenderer.cpp
--- a/Engine/Source/Developer/StandaloneRenderer/Private/Windows/D3D/SlateD3DRenderer.cpp
+++ b/Engine/Source/Developer/StandaloneRenderer/Private/Windows/D3D/SlateD3DRenderer.cpp
@@ -52,13 +52,20 @@ public:
 	{
 	}
 
-	virtual TSharedRef<FSlateFontAtlas> CreateFontAtlas() const override
+	virtual FIntPoint GetAtlasSize() const override
 	{
-		/** Size of each font texture, width and height */
-		const uint32 TextureSize = 1024;
+		return FIntPoint(TextureSize, TextureSize);
+	}
 
+	virtual TSharedRef<FSlateFontAtlas> CreateFontAtlas() const override
+	{
 		return MakeShareable( new FSlateFontAtlasD3D( TextureSize, TextureSize ) );
 	}
+
+private:
+
+	/** Size of each font texture, width and height */
+	static const uint32 TextureSize = 1024;
 };
 
 void FSlateD3DRenderer::Initialize()
@@ -132,6 +139,16 @@ void FSlateD3DRenderer::ReleaseUpdatableTexture(FSlateUpdatableTexture* Texture)
 	delete Texture;
 }
 
+ISlateAtlasProvider* FSlateD3DRenderer::GetTextureAtlasProvider()
+{
+	if( TextureManager.IsValid() )
+	{
+		return TextureManager->GetTextureAtlasProvider();
+	}
+
+	return nullptr;
+}
+
 void FSlateD3DRenderer::Private_CreateViewport( TSharedRef<SWindow> InWindow, const FVector2D &WindowSize )
 {
 	TSharedRef< FGenericWindow > NativeWindow = InWindow->GetNativeWindow().ToSharedRef();
diff --git a/Engine/Source/Developer/StandaloneRenderer/Private/Windows/D3D/SlateD3DRenderer.h b/Engine/Source/Developer/StandaloneRenderer/Private/Windows/D3D/SlateD3DRenderer.h
--- a/Engine/Source/Developer/StandaloneRenderer/Private/Windows/D3D/SlateD3DRenderer.h
+++ b/Engine/Source/Developer/StandaloneRenderer/Private/Windows/D3D/SlateD3DRenderer.h
@@ -58,6 +58,7 @@ public:
 	virtual void LoadStyleResources( const ISlateStyle& Style ) override;
 	virtual FSlateUpdatableTexture* CreateUpdatableTexture(uint32 Width, uint32 Height) override;
 	virtual void ReleaseUpdatableTexture(FSlateUpdatableTexture* Texture) override;
+	virtual ISlateAtlasProvider* GetTextureAtlasProvider() override;
 	
 	void CreateDevice();
 	void CreateDepthStencilBuffer( FSlateD3DViewport& Viewport );
diff --git a/Engine/Source/Developer/StandaloneRenderer/Private/Windows/D3D/SlateD3DTextureManager.cpp b/Engine/Source/Developer/StandaloneRenderer/Private/Windows/D3D/SlateD3DTextureManager.cpp
--- a/Engine/Source/Developer/StandaloneRenderer/Private/Windows/D3D/SlateD3DTextureManager.cpp
+++ b/Engine/Source/Developer/StandaloneRenderer/Private/Windows/D3D/SlateD3DTextureManager.cpp
@@ -47,6 +47,26 @@ FSlateD3DTextureManager::~FSlateD3DTextureManager()
 	}
 }
 
+int32 FSlateD3DTextureManager::GetNumAtlasPages() const
+{
+	return TextureAtlases.Num();
+}
+
+FIntPoint FSlateD3DTextureManager::GetAtlasPageSize() const
+{
+	return FIntPoint(1024, 1024);
+}
+
+FSlateShaderResource* FSlateD3DTextureManager::GetAtlasPageResource(const int32 InIndex) const
+{
+	return TextureAtlases[InIndex]->GetAtlasTexture();
+}
+
+bool FSlateD3DTextureManager::IsAtlasPageResourceAlphaOnly() const
+{
+	return false;
+}
+
 void FSlateD3DTextureManager::LoadUsedTextures()
 {
 	TArray< const FSlateBrush* > Resources;
@@ -174,6 +194,11 @@ FSlateShaderResourceProxy* FSlateD3DTextureManager::GetShaderResource( const FSl
 	return Texture;
 }
 
+ISlateAtlasProvider* FSlateD3DTextureManager::GetTextureAtlasProvider()
+{
+	return this;
+}
+
 bool FSlateD3DTextureManager::LoadTexture( const FSlateBrush& InBrush, uint32& OutWidth, uint32& OutHeight, TArray<uint8>& OutDecodedImage )
 {
 	FString ResourcePath = GetResourcePath( InBrush );
diff --git a/Engine/Source/Developer/StandaloneRenderer/Private/Windows/D3D/SlateD3DTextureManager.h b/Engine/Source/Developer/StandaloneRenderer/Private/Windows/D3D/SlateD3DTextureManager.h
--- a/Engine/Source/Developer/StandaloneRenderer/Private/Windows/D3D/SlateD3DTextureManager.h
+++ b/Engine/Source/Developer/StandaloneRenderer/Private/Windows/D3D/SlateD3DTextureManager.h
@@ -1,17 +1,24 @@
 // Copyright 1998-2014 Epic Games, Inc. All Rights Reserved.
 
-
 #pragma once
 
+#include "TextureAtlas.h"
+
 /**
  * Stores a mapping of texture names to their loaded d3d resource.  Resources are loaded from disk and created on demand when needed                   
  */
-class FSlateD3DTextureManager : public FSlateShaderResourceManager
+class FSlateD3DTextureManager : public ISlateAtlasProvider, public FSlateShaderResourceManager
 {
 public:
 	FSlateD3DTextureManager();
 	virtual ~FSlateD3DTextureManager();
 
+	/** ISlateAtlasProvider */
+	virtual int32 GetNumAtlasPages() const override;
+	virtual FIntPoint GetAtlasPageSize() const override;
+	virtual FSlateShaderResource* GetAtlasPageResource(const int32 InIndex) const override;
+	virtual bool IsAtlasPageResourceAlphaOnly() const override;
+
 	/**
 	 * Loads and creates rendering resources for all used textures.  
 	 * In this implementation all textures must be known at startup time or they will not be found
@@ -26,6 +33,8 @@ public:
 	 */
 	virtual FSlateShaderResourceProxy* GetShaderResource( const FSlateBrush& InBrush ) override;
 
+	virtual ISlateAtlasProvider* GetTextureAtlasProvider() override;
+
 	void CreateTextureNoAtlas( const FSlateBrush& InBrush );
 
 	/**
diff --git a/Engine/Source/Developer/Toolbox/Private/Toolbox.cpp b/Engine/Source/Developer/Toolbox/Private/Toolbox.cpp
--- a/Engine/Source/Developer/Toolbox/Private/Toolbox.cpp
+++ b/Engine/Source/Developer/Toolbox/Private/Toolbox.cpp
@@ -11,6 +11,7 @@
 #include "SNotificationList.h"
 #include "NotificationManager.h"
 #include "STestSuite.h"
+#include "ISlateReflectorModule.h"
 
 IMPLEMENT_MODULE( FToolboxModule, Toolbox );
 
@@ -71,6 +72,15 @@ public:
 				.Text( NSLOCTEXT("DeveloperToolbox", "DisplayTextureAtlases", "Display Texture Atlases") )
 				.OnClicked( this, &SDebugPanel::OnDisplayTextureAtlases )
 			]
+			+ SVerticalBox::Slot()
+			.AutoHeight()
+			.Padding( 4.0f )
+			.HAlign(HAlign_Left)
+			[
+				SNew( SButton )
+				.Text( NSLOCTEXT("DeveloperToolbox", "DisplayFontAtlases", "Display Font Atlases") )
+				.OnClicked( this, &SDebugPanel::OnDisplayFontAtlases )
+			]
 		];
 	}
 	END_SLATE_FUNCTION_BUILD_OPTIMIZATION
@@ -84,7 +94,15 @@ public:
 
 	FReply OnDisplayTextureAtlases()
 	{
-		FSlateApplication::Get().GetRenderer()->DisplayTextureAtlases();
+		static const FName SlateReflectorModuleName("SlateReflector");
+		FModuleManager::LoadModuleChecked<ISlateReflectorModule>(SlateReflectorModuleName).DisplayTextureAtlasVisualizer();
+		return FReply::Handled();
+	}
+
+	FReply OnDisplayFontAtlases()
+	{
+		static const FName SlateReflectorModuleName("SlateReflector");
+		FModuleManager::LoadModuleChecked<ISlateReflectorModule>(SlateReflectorModuleName).DisplayFontAtlasVisualizer();
 		return FReply::Handled();
 	}
 
diff --git a/Engine/Source/Developer/Toolbox/Toolbox.Build.cs b/Engine/Source/Developer/Toolbox/Toolbox.Build.cs
--- a/Engine/Source/Developer/Toolbox/Toolbox.Build.cs
+++ b/Engine/Source/Developer/Toolbox/Toolbox.Build.cs
@@ -20,6 +20,7 @@ public class Toolbox : ModuleRules
                 "InputCore",
 				"Slate",
 				"SlateCore",
+				"SlateReflector",
                 "EditorStyle",
 				"DesktopPlatform",
 				"AppFramework"
diff --git a/Engine/Source/Runtime/Engine/Classes/GameFramework/CheatManager.h b/Engine/Source/Runtime/Engine/Classes/GameFramework/CheatManager.h
--- a/Engine/Source/Runtime/Engine/Classes/GameFramework/CheatManager.h
+++ b/Engine/Source/Runtime/Engine/Classes/GameFramework/CheatManager.h
@@ -247,6 +247,14 @@ class ENGINE_API UCheatManager : public UObject
 	UFUNCTION(exec)
 	virtual void WidgetReflector();
 
+	/** Spawns a Slate texture atlas visualizer in game */
+	UFUNCTION(exec)
+	virtual void TextureAtlasVisualizer();
+
+	/** Spawns a Slate font atlas visualizer in game */
+	UFUNCTION(exec)
+	virtual void FontAtlasVisualizer();
+
 	/** Builds the navigation mesh (or rebuilds it). **/
 	UFUNCTION(exec)
 	virtual void RebuildNavigation();
diff --git a/Engine/Source/Runtime/Engine/Private/CheatManager.cpp b/Engine/Source/Runtime/Engine/Private/CheatManager.cpp
--- a/Engine/Source/Runtime/Engine/Private/CheatManager.cpp
+++ b/Engine/Source/Runtime/Engine/Private/CheatManager.cpp
@@ -876,22 +876,20 @@ void UCheatManager::TestCollisionDistance()
 
 void UCheatManager::WidgetReflector()
 {
-	static TWeakPtr<SWindow> WidgetReflectorWindow;
-	
-	// Only allow one instance open at a time
-	if( !WidgetReflectorWindow.IsValid() )
-	{
-		const TSharedRef<SWindow> ReflectorWindow = SNew(SWindow)
-			.AutoCenter(EAutoCenter::PrimaryWorkArea)
-			.ClientSize(FVector2D(600,400))
-			[
-				FModuleManager::LoadModuleChecked<ISlateReflectorModule>("SlateReflector").GetWidgetReflector()
-			];
-		
-		WidgetReflectorWindow = ReflectorWindow;
-		
-		FSlateApplication::Get().AddWindow( ReflectorWindow );
-	}
+	static const FName SlateReflectorModuleName("SlateReflector");
+	FModuleManager::LoadModuleChecked<ISlateReflectorModule>(SlateReflectorModuleName).DisplayWidgetReflector();
+}
+
+void UCheatManager::TextureAtlasVisualizer()
+{
+	static const FName SlateReflectorModuleName("SlateReflector");
+	FModuleManager::LoadModuleChecked<ISlateReflectorModule>(SlateReflectorModuleName).DisplayTextureAtlasVisualizer();
+}
+
+void UCheatManager::FontAtlasVisualizer()
+{
+	static const FName SlateReflectorModuleName("SlateReflector");
+	FModuleManager::LoadModuleChecked<ISlateReflectorModule>(SlateReflectorModuleName).DisplayFontAtlasVisualizer();
 }
 
 void UCheatManager::RebuildNavigation()
diff --git a/Engine/Source/Runtime/Slate/Private/Framework/Layout/ScrollyZoomy.cpp b/Engine/Source/Runtime/Slate/Private/Framework/Layout/ScrollyZoomy.cpp
--- a/Engine/Source/Runtime/Slate/Private/Framework/Layout/ScrollyZoomy.cpp
+++ b/Engine/Source/Runtime/Slate/Private/Framework/Layout/ScrollyZoomy.cpp
@@ -7,10 +7,11 @@
 /* FScrollyZoomy structors
  *****************************************************************************/
 
-FScrollyZoomy::FScrollyZoomy( )
-	: AmountScrolledWhileRightMouseDown( 0.0f ),
-	  bShowSoftwareCursor( false ),
-	  SoftwareCursorPosition( FVector2D::ZeroVector )
+FScrollyZoomy::FScrollyZoomy( const bool bInUseIntertialScrolling )
+	: AmountScrolledWhileRightMouseDown( 0.0f )
+	, bShowSoftwareCursor( false )
+	, SoftwareCursorPosition( FVector2D::ZeroVector )
+	, bUseIntertialScrolling( bInUseIntertialScrolling )
 { }
 
 
@@ -93,6 +94,12 @@ FReply FScrollyZoomy::OnMouseButtonUp( const TSharedRef<SWidget> MyWidget, const
 			Reply.SetMousePos( BestPositionInPanel );
 		}
 
+		if (!bUseIntertialScrolling)
+		{
+			HorizontalIntertia.ClearScrollVelocity();
+			VerticalIntertia.ClearScrollVelocity();
+		}
+
 		return Reply;
 	}
 
diff --git a/Engine/Source/Runtime/Slate/Public/Framework/Layout/ScrollyZoomy.h b/Engine/Source/Runtime/Slate/Public/Framework/Layout/ScrollyZoomy.h
--- a/Engine/Source/Runtime/Slate/Public/Framework/Layout/ScrollyZoomy.h
+++ b/Engine/Source/Runtime/Slate/Public/Framework/Layout/ScrollyZoomy.h
@@ -37,7 +37,7 @@ class SLATE_API FScrollyZoomy
 public:
 
 	/** Default constructor for FScrollyZoomy */
-	FScrollyZoomy( );
+	FScrollyZoomy( const bool bInUseIntertialScrolling = true );
 
 	/** 
 	 * Should be called every frame to update simulation state
@@ -131,6 +131,9 @@ private:
 	// The current position of the software cursor.
 	FVector2D SoftwareCursorPosition;
 
+	// Is intertial scrolling enabled?
+	bool bUseIntertialScrolling;
+
 	// Tracks simulation state for horizontal scrolling.
 	FInertialScrollManager HorizontalIntertia;
 
diff --git a/Engine/Source/Runtime/SlateCore/Private/Fonts/FontCache.cpp b/Engine/Source/Runtime/SlateCore/Private/Fonts/FontCache.cpp
--- a/Engine/Source/Runtime/SlateCore/Private/Fonts/FontCache.cpp
+++ b/Engine/Source/Runtime/SlateCore/Private/Fonts/FontCache.cpp
@@ -1071,6 +1071,26 @@ FSlateFontCache::~FSlateFontCache()
 
 }
 
+int32 FSlateFontCache::GetNumAtlasPages() const
+{
+	return FontAtlases.Num();
+}
+
+FIntPoint FSlateFontCache::GetAtlasPageSize() const
+{
+	return FontAtlasFactory->GetAtlasSize();
+}
+
+FSlateShaderResource* FSlateFontCache::GetAtlasPageResource(const int32 InIndex) const
+{
+	return FontAtlases[InIndex]->GetSlateTexture();
+}
+
+bool FSlateFontCache::IsAtlasPageResourceAlphaOnly() const
+{
+	return true;
+}
+
 bool FSlateFontCache::AddNewEntry( TCHAR Character, const FSlateFontKey& InKey, FCharacterEntry& OutCharacterEntry ) const
 {
 	float SubFontScalingFactor = 1.0f;
diff --git a/Engine/Source/Runtime/SlateCore/Private/Rendering/ElementBatcher.cpp b/Engine/Source/Runtime/SlateCore/Private/Rendering/ElementBatcher.cpp
--- a/Engine/Source/Runtime/SlateCore/Private/Rendering/ElementBatcher.cpp
+++ b/Engine/Source/Runtime/SlateCore/Private/Rendering/ElementBatcher.cpp
@@ -1127,9 +1127,21 @@ void FSlateElementBatcher::AddViewportElement( const FSlateDrawElement& DrawElem
 
 	TSharedPtr<const ISlateViewport> ViewportPin = InPayload.Viewport.Pin();
 
-	FSlateShaderResource* ViewportResource = ViewportPin.IsValid() ? ViewportPin->GetViewportRenderTargetTexture() : nullptr;
+	FSlateShaderResource* ViewportResource =  nullptr;
+	ESlateShader::Type ShaderType = ESlateShader::Default;
 
-	FSlateElementBatch& ElementBatch = FindBatchForElement( Layer, FShaderParams(), ViewportResource, ESlateDrawPrimitive::TriangleList, ESlateShader::Default, InDrawEffects, DrawFlags, DrawElement.GetScissorRect() );
+	if( ViewportPin.IsValid() )
+	{
+		ViewportResource = ViewportPin->GetViewportRenderTargetTexture();
+
+		if( ViewportPin->IsViewportTextureAlphaOnly() )
+		{
+			// This is a slight hack, but the font shader is the same as the general shader except it reads alpha only textures
+			ShaderType = ESlateShader::Font;
+		}
+	}
+
+	FSlateElementBatch& ElementBatch = FindBatchForElement( Layer, FShaderParams(), ViewportResource, ESlateDrawPrimitive::TriangleList, ShaderType, InDrawEffects, DrawFlags, DrawElement.GetScissorRect() );
 	TArray<FSlateVertex>& BatchVertices = BatchVertexArrays[ElementBatch.VertexArrayIndex];
 	TArray<SlateIndex>& BatchIndices = BatchIndexArrays[ElementBatch.IndexArrayIndex];
 
diff --git a/Engine/Source/Runtime/SlateCore/Private/Rendering/SlateRenderer.cpp b/Engine/Source/Runtime/SlateCore/Private/Rendering/SlateRenderer.cpp
--- a/Engine/Source/Runtime/SlateCore/Private/Rendering/SlateRenderer.cpp
+++ b/Engine/Source/Runtime/SlateCore/Private/Rendering/SlateRenderer.cpp
@@ -39,6 +39,23 @@ bool FSlateRenderer::IsViewportFullscreen( const SWindow& Window ) const
 }
 
 
+ISlateAtlasProvider* FSlateRenderer::GetTextureAtlasProvider()
+{
+	return nullptr;
+}
+
+
+ISlateAtlasProvider* FSlateRenderer::GetFontAtlasProvider()
+{
+	if( FontCache.IsValid() )
+	{
+		return FontCache.Get();
+	}
+
+	return nullptr;
+}
+
+
 /* Global functions
  *****************************************************************************/
 
diff --git a/Engine/Source/Runtime/SlateCore/Public/Fonts/FontCache.h b/Engine/Source/Runtime/SlateCore/Public/Fonts/FontCache.h
--- a/Engine/Source/Runtime/SlateCore/Public/Fonts/FontCache.h
+++ b/Engine/Source/Runtime/SlateCore/Public/Fonts/FontCache.h
@@ -2,10 +2,11 @@
 
 #pragma once
 
+#include "TextureAtlas.h"
+
 struct FCharacterRenderData;
 class FFreeTypeInterface;
 
-
 /** Information for rendering one character */
 struct SLATECORE_API FCharacterEntry
 {
@@ -218,7 +219,7 @@ private:
  * Font caching implementation
  * Caches characters into textures as needed
  */
-class SLATECORE_API FSlateFontCache
+class SLATECORE_API FSlateFontCache : public ISlateAtlasProvider
 {
 public:
 	/**
@@ -230,6 +231,12 @@ public:
 	FSlateFontCache( TSharedRef<ISlateFontAtlasFactory> InFontAtlasFactory );
 	virtual ~FSlateFontCache();
 
+	/** ISlateAtlasProvider */
+	virtual int32 GetNumAtlasPages() const override;
+	virtual FIntPoint GetAtlasPageSize() const override;
+	virtual FSlateShaderResource* GetAtlasPageResource(const int32 InIndex) const override;
+	virtual bool IsAtlasPageResourceAlphaOnly() const override;
+
 	/** 
 	 * Gets information for how to draw all characters in the specified string. Caches characters as they are found
 	 * 
diff --git a/Engine/Source/Runtime/SlateCore/Public/Fonts/FontTypes.h b/Engine/Source/Runtime/SlateCore/Public/Fonts/FontTypes.h
--- a/Engine/Source/Runtime/SlateCore/Public/Fonts/FontTypes.h
+++ b/Engine/Source/Runtime/SlateCore/Public/Fonts/FontTypes.h
@@ -122,5 +122,6 @@ public:
 class SLATECORE_API ISlateFontAtlasFactory
 {
 public:
+	virtual FIntPoint GetAtlasSize() const = 0;
 	virtual TSharedRef<class FSlateFontAtlas> CreateFontAtlas() const = 0;
 };
\ No newline at end of file
diff --git a/Engine/Source/Runtime/SlateCore/Public/Rendering/RenderingCommon.h b/Engine/Source/Runtime/SlateCore/Public/Rendering/RenderingCommon.h
--- a/Engine/Source/Runtime/SlateCore/Public/Rendering/RenderingCommon.h
+++ b/Engine/Source/Runtime/SlateCore/Public/Rendering/RenderingCommon.h
@@ -229,6 +229,14 @@ public:
 	 */
 	virtual class FSlateShaderResource* GetViewportRenderTargetTexture() const = 0;
 
+	/**
+	 * Does the texture returned by GetViewportRenderTargetTexture only have an alpha channel?
+	 */
+	virtual bool IsViewportTextureAlphaOnly() const
+	{
+		return false;
+	}
+
 	/**
 	 * Performs any ticking necessary by this handle                   
 	 */
diff --git a/Engine/Source/Runtime/SlateCore/Public/Rendering/SlateRenderer.h b/Engine/Source/Runtime/SlateCore/Public/Rendering/SlateRenderer.h
--- a/Engine/Source/Runtime/SlateCore/Public/Rendering/SlateRenderer.h
+++ b/Engine/Source/Runtime/SlateCore/Public/Rendering/SlateRenderer.h
@@ -8,16 +8,16 @@ class FRHITexture2D;
 class ISlateStyle;
 class FSlateDrawBuffer;
 class FSlateUpdatableTexture;
+class ISlateAtlasProvider;
 
 struct FSlateBrush;
 
 typedef FRHITexture2D* FTexture2DRHIParamRef;
 
-
 /**
  * Abstract base class for Slate renderers.
  */
-class FSlateRenderer
+class SLATECORE_API FSlateRenderer
 {
 public:
 
@@ -153,7 +153,7 @@ public:
 	/**
 	 * Flushes all cached data from the font cache
 	 */
-	void SLATECORE_API FlushFontCache();
+	void FlushFontCache();
 
 	/**
 	 * Gives the renderer a chance to wait for any render commands to be completed before returning/
@@ -176,9 +176,6 @@ public:
 	 */
 	virtual void LoadStyleResources( const ISlateStyle& Style ) {}
 
-	/** Creates a window to visualize the texture atlases */
-	virtual void DisplayTextureAtlases() {}
-
 	/** Releases a specific resource */
 	virtual void ReleaseDynamicResource( const FSlateBrush& Brush ) = 0;
 
@@ -188,7 +185,7 @@ public:
 	 * @Window	The window to check for fullscreen
 	 * @return true if the window's viewport should be fullscreen
 	 */
-	bool SLATECORE_API IsViewportFullscreen( const SWindow& Window ) const;
+	bool IsViewportFullscreen( const SWindow& Window ) const;
 
 	/** Returns whether shaders that Slate depends on have been compiled. */
 	virtual bool AreShadersInitialized() const { return true; }
@@ -228,6 +225,16 @@ public:
 	 */
 	virtual void ReleaseUpdatableTexture(FSlateUpdatableTexture* Texture) = 0;
 
+	/**
+	 * Returns the way to access the texture atlas information for this renderer
+	 */
+	virtual ISlateAtlasProvider* GetTextureAtlasProvider();
+
+	/**
+	 * Returns the way to access the font atlas information for this renderer
+	 */
+	virtual ISlateAtlasProvider* GetFontAtlasProvider();
+
 private:
 
 	// Non-copyable
diff --git a/Engine/Source/Runtime/SlateCore/Public/Textures/TextureAtlas.h b/Engine/Source/Runtime/SlateCore/Public/Textures/TextureAtlas.h
--- a/Engine/Source/Runtime/SlateCore/Public/Textures/TextureAtlas.h
+++ b/Engine/Source/Runtime/SlateCore/Public/Textures/TextureAtlas.h
@@ -214,3 +214,25 @@ protected:
 	 */
 	ESlateTextureAtlasOwnerThread AtlasOwnerThread;
 };
+
+/**
+ * Interface to allow the Slate atlas visualizer to query atlas page information for an atlas provider
+ */
+class ISlateAtlasProvider
+{
+public:
+	/** Virtual destructor */
+	virtual ~ISlateAtlasProvider() {}
+
+	/** Get the number of atlas pages this atlas provider has available when calling GetAtlasPageResource */
+	virtual int32 GetNumAtlasPages() const = 0;
+
+	/** Get the size of each atlas page */
+	virtual FIntPoint GetAtlasPageSize() const = 0;
+
+	/** Get the page resource for the given index (verify with GetNumAtlasPages) */ 
+	virtual FSlateShaderResource* GetAtlasPageResource(const int32 InIndex) const = 0;
+
+	/** Do the atlas page resources only contain alpha information? This affects how the atlas visualizer will sample them */
+	virtual bool IsAtlasPageResourceAlphaOnly() const = 0;
+};
diff --git a/Engine/Source/Runtime/SlateCore/Public/Textures/TextureManager.h b/Engine/Source/Runtime/SlateCore/Public/Textures/TextureManager.h
--- a/Engine/Source/Runtime/SlateCore/Public/Textures/TextureManager.h
+++ b/Engine/Source/Runtime/SlateCore/Public/Textures/TextureManager.h
@@ -2,6 +2,8 @@
 
 #pragma once
 
+class ISlateAtlasProvider;
+
 /**
  * Represents a new texture that has been loaded but no resource created for it
  */
@@ -48,6 +50,11 @@ public:
 	 */
 	virtual FSlateShaderResourceProxy* GetShaderResource( const FSlateBrush& InBrush ) = 0;
 
+	/**
+	 * Returns the way to access the texture atlas information from this resource manager
+	 */
+	virtual ISlateAtlasProvider* GetTextureAtlasProvider() = 0;
+
 protected:
 
 	void ClearTextureMap()
diff --git a/Engine/Source/Runtime/SlateRHIRenderer/Private/SlateRHIRenderer.cpp b/Engine/Source/Runtime/SlateRHIRenderer/Private/SlateRHIRenderer.cpp
--- a/Engine/Source/Runtime/SlateRHIRenderer/Private/SlateRHIRenderer.cpp
+++ b/Engine/Source/Runtime/SlateRHIRenderer/Private/SlateRHIRenderer.cpp
@@ -1095,24 +1095,6 @@ void FSlateRHIRenderer::LoadStyleResources( const ISlateStyle& Style )
 	}
 }
 
-void FSlateRHIRenderer::DisplayTextureAtlases()
-{
-	TSharedRef<SWindow> Window = SNew(SWindow)
-		.SizingRule( ESizingRule::Autosized )
-		.SupportsMaximize(false)
-		.SupportsMinimize(false)
-		.Title( FText::GetEmpty() )
-		[
-			SNew( SBorder )
-			.BorderImage( FCoreStyle::Get().GetBrush("ToolPanel.GroupBorder") )
-			[
-				ResourceManager->CreateTextureDisplayWidget()
-			]
-		];
-
-	FSlateApplication::Get().AddWindow( Window );
-}
-
 void FSlateRHIRenderer::ReleaseDynamicResource( const FSlateBrush& InBrush )
 {
 	ensure( IsInGameThread() );
@@ -1182,6 +1164,16 @@ void FSlateRHIRenderer::ReleaseUpdatableTexture(FSlateUpdatableTexture* Texture)
 	delete Texture;
 }
 
+ISlateAtlasProvider* FSlateRHIRenderer::GetTextureAtlasProvider()
+{
+	if( ResourceManager.IsValid() )
+	{
+		return ResourceManager->GetTextureAtlasProvider();
+	}
+
+	return nullptr;
+}
+
 bool FSlateRHIRenderer::AreShadersInitialized() const
 {
 #if WITH_EDITORONLY_DATA
diff --git a/Engine/Source/Runtime/SlateRHIRenderer/Private/SlateRHIRenderer.h b/Engine/Source/Runtime/SlateRHIRenderer/Private/SlateRHIRenderer.h
--- a/Engine/Source/Runtime/SlateRHIRenderer/Private/SlateRHIRenderer.h
+++ b/Engine/Source/Runtime/SlateRHIRenderer/Private/SlateRHIRenderer.h
@@ -168,6 +168,7 @@ public:
 	virtual void SetColorVisionDeficiencyType( uint32 Type ) override;
 	virtual FSlateUpdatableTexture* CreateUpdatableTexture(uint32 Width, uint32 Height) override;
 	virtual void ReleaseUpdatableTexture(FSlateUpdatableTexture* Texture) override;
+	virtual ISlateAtlasProvider* GetTextureAtlasProvider() override;
 
 	/** Draws windows from a FSlateDrawBuffer on the render thread */
 	void DrawWindow_RenderThread(FRHICommandListImmediate& RHICmdList, const FSlateRHIRenderer::FViewportInfo& ViewportInfo, const FSlateWindowElementList& WindowElementList, bool bLockToVsync);
@@ -202,11 +203,6 @@ public:
 
 	virtual void LoadStyleResources( const ISlateStyle& Style ) override;
 
-	/**
-	 * Creates a window with an atlas visualizer inside it
-	 */
-	virtual void DisplayTextureAtlases();
-
 	/**
 	 * Returns the viewport RHI reference for the provided window
 	 *
diff --git a/Engine/Source/Runtime/SlateRHIRenderer/Private/SlateRHIRendererModule.cpp b/Engine/Source/Runtime/SlateRHIRenderer/Private/SlateRHIRendererModule.cpp
--- a/Engine/Source/Runtime/SlateRHIRenderer/Private/SlateRHIRendererModule.cpp
+++ b/Engine/Source/Runtime/SlateRHIRenderer/Private/SlateRHIRendererModule.cpp
@@ -27,6 +27,11 @@ public:
 	{
 	}
 
+	virtual FIntPoint GetAtlasSize() const override
+	{
+		return FIntPoint(AtlasSize, AtlasSize);
+	}
+
 	virtual TSharedRef<FSlateFontAtlas> CreateFontAtlas() const override
 	{
 		return MakeShareable(new FSlateFontAtlasRHI(AtlasSize, AtlasSize));
diff --git a/Engine/Source/Runtime/SlateRHIRenderer/Private/SlateRHIResourceManager.cpp b/Engine/Source/Runtime/SlateRHIRenderer/Private/SlateRHIResourceManager.cpp
--- a/Engine/Source/Runtime/SlateRHIRenderer/Private/SlateRHIResourceManager.cpp
+++ b/Engine/Source/Runtime/SlateRHIRenderer/Private/SlateRHIResourceManager.cpp
@@ -156,6 +156,26 @@ FSlateRHIResourceManager::~FSlateRHIResourceManager()
 	}
 }
 
+int32 FSlateRHIResourceManager::GetNumAtlasPages() const
+{
+	return TextureAtlases.Num();
+}
+
+FIntPoint FSlateRHIResourceManager::GetAtlasPageSize() const
+{
+	return FIntPoint(1024, 1024);
+}
+
+FSlateShaderResource* FSlateRHIResourceManager::GetAtlasPageResource(const int32 InIndex) const
+{
+	return TextureAtlases[InIndex]->GetAtlasTexture();
+}
+
+bool FSlateRHIResourceManager::IsAtlasPageResourceAlphaOnly() const
+{
+	return false;
+}
+
 void FSlateRHIResourceManager::CreateTextures( const TArray< const FSlateBrush* >& Resources )
 {
 	DECLARE_SCOPE_CYCLE_COUNTER(TEXT("Loading Slate Textures"), STAT_Slate, STATGROUP_LoadTime);
@@ -383,6 +403,11 @@ FSlateShaderResourceProxy* FSlateRHIResourceManager::GetShaderResource( const FS
 	return Texture;
 }
 
+ISlateAtlasProvider* FSlateRHIResourceManager::GetTextureAtlasProvider()
+{
+	return this;
+}
+
 TSharedPtr<FSlateDynamicTextureResource> FSlateRHIResourceManager::MakeDynamicTextureResource( FName ResourceName, uint32 Width, uint32 Height, const TArray< uint8 >& Bytes )
 {
 	// Make storage for the image
@@ -766,198 +791,3 @@ void FSlateRHIResourceManager::ReloadTextures()
 	// Reload everythng
 	LoadUsedTextures();
 }
-
-uint32 FSlateRHIResourceManager::GetNumTextureAtlases() const
-{
-	return TextureAtlases.Num();
-}
-
-FSlateShaderResource* FSlateRHIResourceManager::GetTextureAtlas( uint32 Index )
-{
-	return TextureAtlases[Index]->GetAtlasTexture();
-}
-
-struct FAtlasPage
-{
-	FAtlasPage( FSlateShaderResource* InTexture, uint32 InIndex )
-		: Texture( InTexture )
-		, Index( InIndex )
-	{}
-
-	FSlateShaderResource* Texture;
-	uint32 Index;
-};
-
-class FAtlasVisualizer : public ISlateViewport, public TSharedFromThis<FAtlasVisualizer>
-{
-public:
-	FAtlasVisualizer( FSlateRHIResourceManager& InTextureManager )
-		: TextureManager( InTextureManager )
-		, SelectedAtlasPage( NULL )
-	{}
-
-	virtual FIntPoint GetSize() const override { return FIntPoint(1024,1024); }
-
-	virtual FSlateShaderResource* GetViewportRenderTargetTexture() const override
-	{
-		if( SelectedAtlasPage.IsValid() )
-		{
-			return SelectedAtlasPage.Pin()->Texture;
-		}
-		else
-		{
-			return NULL;
-		}
-	}
-
-	virtual bool RequiresVsync() const override { return false; }
-
-	TSharedRef<SWidget> MakeVisualizerWidget()
-	{
-		bDisplayCheckerboard = false;
-		AtlasPages.Empty();
-
-		for( uint32 AtlasIndex = 0; AtlasIndex < TextureManager.GetNumTextureAtlases(); ++AtlasIndex )
-		{
-			AtlasPages.Add( MakeShareable( new FAtlasPage( TextureManager.GetTextureAtlas(AtlasIndex), AtlasIndex ) ) );
-		}
-
-		SelectedAtlasPage = AtlasPages.Num() > 0 ? AtlasPages[0] : NULL;
-		TSharedPtr<SViewport> Viewport;
-
-		TSharedRef<SWidget> Widget =
-		SNew( SVerticalBox )
-		+ SVerticalBox::Slot()
-		.Padding(4.0f)
-		.HAlign(HAlign_Left)
-		.AutoHeight()
-		[
-			SNew( SHorizontalBox )
-			+ SHorizontalBox::Slot()
-			.VAlign( VAlign_Center )
-			.AutoWidth()
-			.Padding(0.0,2.0f,2.0f,2.0f)
-			[
-				SNew( STextBlock )
-				.Text( NSLOCTEXT("TextureManagerVisualizer", "SelectAPageLabel", "Select a page") )
-			]
-			+ SHorizontalBox::Slot()
-			.AutoWidth()
-			.Padding(2.0f)
-			.VAlign( VAlign_Center )
-			[
-				SNew( SComboBox< TSharedPtr<FAtlasPage> > )
-				.OptionsSource( &AtlasPages )
-				.OnGenerateWidget( this, &FAtlasVisualizer::OnGenerateWidgetForCombo )
-				.OnSelectionChanged( this, &FAtlasVisualizer::OnAtlasPageChanged )
-				.InitiallySelectedItem( SelectedAtlasPage.IsValid() ? SelectedAtlasPage.Pin() : NULL )
-				.Content()
-				[
-					SNew( STextBlock )
-					.Text( this, &FAtlasVisualizer::OnGetSelectedItemText )
-				]
-			]
-			+ SHorizontalBox::Slot()
-			.Padding(20.0f,2.0f)
-			.AutoWidth()
-			.VAlign( VAlign_Center )
-			[
-				SNew( SCheckBox )
-				.OnCheckStateChanged( this, &FAtlasVisualizer::OnDisplayCheckerboardStateChanged )
-				.IsChecked( this, &FAtlasVisualizer::OnGetCheckerboardState )
-				.Content()
-				[
-					SNew( STextBlock )
-					.Text( NSLOCTEXT("TextureManagerVisualizer", "DisplayCheckerboardCheckboxLabel", "Display Checkerboard") )
-				]
-			]
-		]
-		+ SVerticalBox::Slot()
-		.Padding(2.0f)
-		.AutoHeight()
-		[
-			SNew( SBox )
-			.WidthOverride(1024)
-			.HeightOverride(1024)
-			[
-				SNew( SOverlay )
-				+ SOverlay::Slot()
-				[
-					SNew( SImage )
-					.Visibility( this, &FAtlasVisualizer::OnGetCheckerboardVisibility )
-					.Image( FCoreStyle::Get().GetBrush("Checkerboard") )
-				]
-				+ SOverlay::Slot()
-				[
-					SAssignNew( Viewport, SViewport )
-					.IgnoreTextureAlpha(false)
-					.EnableBlending(true)
-				]
-			]
-		];
-
-		Viewport->SetViewportInterface( SharedThis( this ) );
-
-		return Widget;
-
-	}
-
-private:
-
-	void OnDisplayCheckerboardStateChanged( ESlateCheckBoxState::Type NewState )
-	{
-		bDisplayCheckerboard = NewState == ESlateCheckBoxState::Checked;
-	}
-
-	ESlateCheckBoxState::Type OnGetCheckerboardState() const
-	{
-		return bDisplayCheckerboard ? ESlateCheckBoxState::Checked : ESlateCheckBoxState::Unchecked;
-	}
-
-	EVisibility OnGetCheckerboardVisibility() const
-	{
-		return bDisplayCheckerboard ? EVisibility::Visible : EVisibility::Collapsed;
-	}
-
-	FString OnGetSelectedItemText() const
-	{
-		if( SelectedAtlasPage.IsValid() )
-		{
-			return FString::Printf( TEXT("Page %d"), SelectedAtlasPage.Pin()->Index );
-		}
-		else
-		{
-			return TEXT("Select a page");
-		}
-	}
-
-	void OnAtlasPageChanged( TSharedPtr<FAtlasPage> AtlasPage, ESelectInfo::Type SelectionType )
-	{
-		SelectedAtlasPage = AtlasPage;
-	}
-
-	TSharedRef<SWidget> OnGenerateWidgetForCombo( TSharedPtr<FAtlasPage> AtlasPage )
-	{
-		return 
-		SNew( STextBlock )
-		.Text( FString::Printf( TEXT("Page %d"), AtlasPage->Index ) );
-	}
-private:
-	TArray< TSharedPtr<FAtlasPage> > AtlasPages;
-	FSlateRHIResourceManager& TextureManager;
-	TWeakPtr<FAtlasPage> SelectedAtlasPage;
-	bool bDisplayCheckerboard;
-};
-
-
-
-TSharedRef<SWidget> FSlateRHIResourceManager::CreateTextureDisplayWidget()
-{
-	static TSharedPtr<FAtlasVisualizer> AtlasVisualizer;
-	if( !AtlasVisualizer.IsValid() )
-	{
-		AtlasVisualizer = MakeShareable( new FAtlasVisualizer( *this ) );
-	}
-
-	return AtlasVisualizer->MakeVisualizerWidget();
-}
diff --git a/Engine/Source/Runtime/SlateRHIRenderer/Private/SlateRHIResourceManager.h b/Engine/Source/Runtime/SlateRHIRenderer/Private/SlateRHIResourceManager.h
--- a/Engine/Source/Runtime/SlateRHIRenderer/Private/SlateRHIResourceManager.h
+++ b/Engine/Source/Runtime/SlateRHIRenderer/Private/SlateRHIResourceManager.h
@@ -2,6 +2,8 @@
 
 #pragma once
 
+#include "TextureAtlas.h"
+
 class FSlateDynamicTextureResource;
 class FSlateUTextureResource;
 class FSlateMaterialResource;
@@ -50,12 +52,18 @@ private:
 /**
  * Stores a mapping of texture names to their RHI texture resource               
  */
-class FSlateRHIResourceManager : public FSlateShaderResourceManager
+class FSlateRHIResourceManager : public ISlateAtlasProvider, public FSlateShaderResourceManager
 {
 public:
 	FSlateRHIResourceManager();
 	virtual ~FSlateRHIResourceManager();
 
+	/** ISlateAtlasProvider */
+	virtual int32 GetNumAtlasPages() const override;
+	virtual FIntPoint GetAtlasPageSize() const override;
+	virtual FSlateShaderResource* GetAtlasPageResource(const int32 InIndex) const override;
+	virtual bool IsAtlasPageResourceAlphaOnly() const override;
+
 	/**
 	 * Loads and creates rendering resources for all used textures.  
 	 * In this implementation all textures must be known at startup time or they will not be found
@@ -78,7 +86,8 @@ public:
 
 	/** FSlateShaderResourceManager interface */
 	virtual FSlateShaderResourceProxy* GetShaderResource( const FSlateBrush& InBrush ) override;
-	
+	virtual ISlateAtlasProvider* GetTextureAtlasProvider() override;
+
 	/**
 	 * Makes a dynamic texture resource and begins use of it
 	 *
@@ -132,16 +141,6 @@ public:
 	 */
 	void ReloadTextures();
 
-	/** @return The number of texture atlases in the manager */
-	uint32 GetNumTextureAtlases() const;
-
-	/** @return The atlas texture at a given index */
-	FSlateShaderResource* GetTextureAtlas( uint32 Index );
-
-	/**
-	 * Creates an atlas visualizer widget
-	 */
-	TSharedRef<SWidget> CreateTextureDisplayWidget();
 private:
 	/**
 	 * Deletes resources created by the manager
