Editor - The widget reflector now caches most of the information off a widget so that if the widget is destroyed after you lose focus (menus), you can still navigate to the source file where it was defined without that data disappearing.

diff --git a/Engine/Source/Developer/SlateReflector/Private/Widgets/SWidgetReflector.cpp b/Engine/Source/Developer/SlateReflector/Private/Widgets/SWidgetReflector.cpp
--- a/Engine/Source/Developer/SlateReflector/Private/Widgets/SWidgetReflector.cpp
+++ b/Engine/Source/Developer/SlateReflector/Private/Widgets/SWidgetReflector.cpp
@@ -73,34 +73,36 @@ void SWidgetReflector::Construct( const FArguments& InArgs )
 						SNew(SHorizontalBox)
 					
 						+ SHorizontalBox::Slot()
-							.AutoWidth()
-							.VAlign(VAlign_Center)
-							[
-								SNew(STextBlock)
-									.Text(LOCTEXT("AppScale", "Application Scale: ").ToString())
-							]
+						.AutoWidth()
+						.VAlign(VAlign_Center)
+						[
+							SNew(STextBlock)
+								.Text(LOCTEXT("AppScale", "Application Scale: ").ToString())
+						]
 
 						+ SHorizontalBox::Slot()
-							.AutoWidth()
+						.AutoWidth()
+						[
+							SNew(SBox)
+							.MinDesiredWidth(100)
+							.MaxDesiredWidth(250)
 							[
-								SNew(SBox)
-								.MinDesiredWidth(100)
-								.MaxDesiredWidth(250)
-								[
-									SNew(SSpinBox<float>)
-										.Value(this, &SWidgetReflector::HandleAppScaleSliderValue)
-										.MinValue(0.1f)
-										.MaxValue(3.0f)
-										.Delta(0.01f)
-										.OnValueChanged(this, &SWidgetReflector::HandleAppScaleSliderChanged)
-								]
+								SNew(SSpinBox<float>)
+									.Value(this, &SWidgetReflector::HandleAppScaleSliderValue)
+									.MinValue(0.1f)
+									.MaxValue(3.0f)
+									.Delta(0.01f)
+									.OnValueChanged(this, &SWidgetReflector::HandleAppScaleSliderChanged)
 							]
+						]
+
 						+ SHorizontalBox::Slot()
 						.FillWidth(1.0f)
 						[
 							SNew(SSpacer)
 						]
-						+SHorizontalBox::Slot()
+
+						+ SHorizontalBox::Slot()
 						.AutoWidth()
 						.Padding(FMargin(5.0f, 0.0f))
 						[
@@ -108,7 +110,8 @@ void SWidgetReflector::Construct( const FArguments& InArgs )
 							.Text(LOCTEXT("DisplayTextureAtlases", "Display Texture Atlases"))
 							.OnClicked(this, &SWidgetReflector::HandleDisplayTextureAtlases)
 						]
-						+SHorizontalBox::Slot()
+
+						+ SHorizontalBox::Slot()
 						.AutoWidth()
 						.Padding(FMargin(5.0f, 0.0f))
 						[
@@ -119,39 +122,39 @@ void SWidgetReflector::Construct( const FArguments& InArgs )
 					]
 
 				+ SVerticalBox::Slot()
-					.AutoHeight()
-					[
-						SNew(SHorizontalBox)
+				.AutoHeight()
+				[
+					SNew(SHorizontalBox)
 
-						+ SHorizontalBox::Slot()
-							.AutoWidth()
-							.Padding(5.0f)
-							[
-								// Check box that controls LIVE MODE
-								SNew(SCheckBox)
-									.IsChecked(this, &SWidgetReflector::HandleFocusCheckBoxIsChecked)
-									.OnCheckStateChanged(this, &SWidgetReflector::HandleFocusCheckBoxCheckedStateChanged)
-									[
-										SNew(STextBlock)
-											.Text(LOCTEXT("ShowFocus", "Show Focus").ToString())
-									]
-							]
+					+ SHorizontalBox::Slot()
+					.AutoWidth()
+					.Padding(5.0f)
+					[
+						// Check box that controls LIVE MODE
+						SNew(SCheckBox)
+						.IsChecked(this, &SWidgetReflector::HandleFocusCheckBoxIsChecked)
+						.OnCheckStateChanged(this, &SWidgetReflector::HandleFocusCheckBoxCheckedStateChanged)
+						[
+							SNew(STextBlock)
+							.Text(LOCTEXT("ShowFocus", "Show Focus").ToString())
+						]
+					]
 
-						+ SHorizontalBox::Slot()
-							.AutoWidth()
-							.Padding(5.0f)
-							[
-								// Check box that controls PICKING A WIDGET TO INSPECT
-								SNew(SButton)
-									.OnClicked(this, &SWidgetReflector::HandlePickButtonClicked)
-									.ButtonColorAndOpacity(this, &SWidgetReflector::HandlePickButtonColorAndOpacity)
-									[
-										SNew(STextBlock)
-											.Text(this, &SWidgetReflector::HandlePickButtonText)
-									]
-							]
+					+ SHorizontalBox::Slot()
+					.AutoWidth()
+					.Padding(5.0f)
+					[
+						// Check box that controls PICKING A WIDGET TO INSPECT
+						SNew(SButton)
+						.OnClicked(this, &SWidgetReflector::HandlePickButtonClicked)
+						.ButtonColorAndOpacity(this, &SWidgetReflector::HandlePickButtonColorAndOpacity)
+						[
+							SNew(STextBlock)
+							.Text(this, &SWidgetReflector::HandlePickButtonText)
+						]
+					]
 
-					+SHorizontalBox::Slot()
+					+ SHorizontalBox::Slot()
 					.AutoWidth()
 					.Padding(5.0f)
 					[
@@ -173,7 +176,8 @@ void SWidgetReflector::Construct( const FArguments& InArgs )
 						#endif
 						.OnClicked_Static( &ToggleSlateStats )
 					]
-					+SHorizontalBox::Slot()
+
+					+ SHorizontalBox::Slot()
 					.AutoWidth()
 					.Padding(5.0f)
 					[
@@ -181,7 +185,8 @@ void SWidgetReflector::Construct( const FArguments& InArgs )
 						.OnClicked( this, &SWidgetReflector::CopyStatsToClipboard )
 						.Text( LOCTEXT("CopyStatsToClipboard", "Copy Stats") )
 					]
-					+SHorizontalBox::Slot()
+
+					+ SHorizontalBox::Slot()
 					.AutoWidth()
 					.Padding(5.0f)
 					[
diff --git a/Engine/Source/Developer/SlateReflector/Private/Widgets/SWidgetReflectorTreeWidgetItem.cpp b/Engine/Source/Developer/SlateReflector/Private/Widgets/SWidgetReflectorTreeWidgetItem.cpp
--- a/Engine/Source/Developer/SlateReflector/Private/Widgets/SWidgetReflectorTreeWidgetItem.cpp
+++ b/Engine/Source/Developer/SlateReflector/Private/Widgets/SWidgetReflectorTreeWidgetItem.cpp
@@ -98,8 +98,6 @@ END_SLATE_FUNCTION_BUILD_OPTIMIZATION
 
 FString SReflectorTreeWidgetItem::GetReadableLocation() const
 {
-	FString ReadableLocation;
-
 	TWeakPtr<SWidget> Widget = WidgetInfo.Get()->Widget;
 	if ( Widget.IsValid() )
 	{
@@ -107,15 +105,11 @@ FString SReflectorTreeWidgetItem::GetReadableLocation() const
 		TSharedPtr<FReflectionMetaData> MetaData = SafeWidget->GetMetaData<FReflectionMetaData>();
 		if ( MetaData.IsValid() )
 		{
-			ReadableLocation = MetaData->Asset->GetName() + TEXT(" [") + MetaData->Name.ToString() + TEXT("]");
-		}
-		else
-		{
-			ReadableLocation = SafeWidget->GetReadableLocation();
+			return MetaData->Asset->GetName() + TEXT(" [") + MetaData->Name.ToString() + TEXT("]");
 		}
 	}
 
-	return ReadableLocation;
+	return CachedReadableLocation;
 }
 
 void SReflectorTreeWidgetItem::HandleHyperlinkNavigate()
@@ -130,14 +124,13 @@ void SReflectorTreeWidgetItem::HandleHyperlinkNavigate()
 			if ( OnAccessAsset.IsBound() )
 			{
 				OnAccessAsset.Execute(MetaData->Asset.Get());
+				return;
 			}
 		}
-		else
-		{
-			if ( OnAccessSourceCode.IsBound() )
-			{
-				OnAccessSourceCode.Execute(GetWidgetFile().ToString(), GetWidgetLineNumber(), 0);
-			}
-		}
+	}
+
+	if ( OnAccessSourceCode.IsBound() )
+	{
+		OnAccessSourceCode.Execute(GetWidgetFile().ToString(), GetWidgetLineNumber(), 0);
 	}
 }
diff --git a/Engine/Source/Developer/SlateReflector/Private/Widgets/SWidgetReflectorTreeWidgetItem.h b/Engine/Source/Developer/SlateReflector/Private/Widgets/SWidgetReflectorTreeWidgetItem.h
--- a/Engine/Source/Developer/SlateReflector/Private/Widgets/SWidgetReflectorTreeWidgetItem.h
+++ b/Engine/Source/Developer/SlateReflector/Private/Widgets/SWidgetReflectorTreeWidgetItem.h
@@ -36,6 +36,17 @@ public:
 		this->OnAccessSourceCode = InArgs._SourceCodeAccessor;
 		this->OnAccessAsset = InArgs._AssetAccessor;
 
+		TSharedPtr<SWidget> Widget = WidgetInfo.Get()->Widget.Pin();
+		if ( Widget.IsValid() )
+		{
+			CachedWidgetType = FText::FromString(Widget->GetTypeAsString());
+			CachedWidgetFile = FText::FromString(Widget->GetCreatedInFile());
+			CachedWidgetLineNumber =  Widget->GetCreatedInLineNumber();
+			CachedWidgetVisibility = FText::FromString(Widget->GetVisibility().ToString());
+
+			CachedReadableLocation = Widget->GetReadableLocation();
+		}
+
 		SMultiColumnTableRow< TSharedPtr<FReflectorNode> >::Construct( SMultiColumnTableRow< TSharedPtr<FReflectorNode> >::FArguments().Padding(1), InOwnerTableView );
 	}
 
@@ -50,9 +61,7 @@ protected:
 	/** @return String representation of the widget we are visualizing */
 	FText GetWidgetType() const
 	{
-		return WidgetInfo.Get()->Widget.IsValid()
-			? FText::FromString(WidgetInfo.Get()->Widget.Pin()->GetTypeAsString())
-			: NSLOCTEXT("SWidgetReflector", "NullWidget", "Null Widget");
+		return CachedWidgetType;
 	}
 	
 	virtual FString GetReadableLocation() const override;
@@ -64,24 +73,17 @@ protected:
 
 	FText GetWidgetFile() const
 	{
-		return WidgetInfo.Get()->Widget.IsValid()
-			? FText::FromString(WidgetInfo.Get()->Widget.Pin()->GetCreatedInFile())
-			: FText::GetEmpty();
+		return CachedWidgetFile;
 	}
 
 	int32 GetWidgetLineNumber() const
 	{
-		return WidgetInfo.Get()->Widget.IsValid()
-			? WidgetInfo.Get()->Widget.Pin()->GetCreatedInLineNumber()
-			: 0;
+		return CachedWidgetLineNumber;
 	}
 
 	FText GetVisibilityAsString() const
 	{
-		TSharedPtr<SWidget> TheWidget = WidgetInfo.Get()->Widget.Pin();
-		return TheWidget.IsValid()
-			? FText::FromString(TheWidget->GetVisibility().ToString())
-			: FText::GetEmpty();
+		return CachedWidgetVisibility;
 	}
 
 	/** @return The tint of the reflector node */
@@ -97,6 +99,13 @@ private:
 	/** The info about the widget that we are visualizing. */
 	TAttribute< TSharedPtr<FReflectorNode> > WidgetInfo;
 
+	FText CachedWidgetType;
+	FText CachedWidgetFile;
+	int32 CachedWidgetLineNumber;
+	FText CachedWidgetVisibility;
+
+	FString CachedReadableLocation;
+
 	FAccessSourceCode OnAccessSourceCode;
 
 	FAccessAsset OnAccessAsset;
