Mergin CL 2411459, 284248 from UE4-UT/... -> UE4/...
   Slate : added DemoMode visualization of slate mouse events - draw a fuchsia ring around event location.

diff --git a/Engine/Source/Developer/SlateReflector/Private/Widgets/SWidgetReflector.cpp b/Engine/Source/Developer/SlateReflector/Private/Widgets/SWidgetReflector.cpp
--- a/Engine/Source/Developer/SlateReflector/Private/Widgets/SWidgetReflector.cpp
+++ b/Engine/Source/Developer/SlateReflector/Private/Widgets/SWidgetReflector.cpp
@@ -60,6 +60,10 @@ void SWidgetReflector::Construct( const FArguments& InArgs )
 	bShowFocus = false;
 	bIsPicking = false;
 
+	bEnableDemoMode = false;
+	LastMouseClickTime = -1;
+	CursorPingPosition = FVector2D::ZeroVector;
+
 	this->ChildSlot
 	[
 		SNew(SBorder)
@@ -210,6 +214,30 @@ void SWidgetReflector::Construct( const FArguments& InArgs )
 							]
 						]
 					]
+					+ SHorizontalBox::Slot()
+					.AutoWidth()
+					.Padding(5.0f)
+					[
+						SNew(SCheckBox)
+						.Style(FCoreStyle::Get(), "ToggleButtonCheckbox")
+						.IsChecked_Lambda([&]()
+						{
+							return bEnableDemoMode ? ECheckBoxState::Checked : ECheckBoxState::Unchecked;
+						})
+						.OnCheckStateChanged_Lambda([&](const ECheckBoxState NewState)
+						{
+							bEnableDemoMode = (NewState == ECheckBoxState::Checked) ? true : false;
+						})
+						[
+							SNew(SBox)
+							.VAlign(VAlign_Center)
+							.HAlign(HAlign_Center)
+							[
+								SNew(STextBlock)
+								.Text(LOCTEXT("EnableDemoMode", "Demo Mode"))
+							]
+						]
+					]
 				]
 
 				+ SVerticalBox::Slot()
@@ -285,6 +313,16 @@ void SWidgetReflector::Tick( const FGeometry& AllottedGeometry, const double InC
 
 void SWidgetReflector::OnEventProcessed( const FInputEvent& Event, const FReplyBase& InReply )
 {
+	if ( Event.IsPointerEvent() )
+	{
+		const FPointerEvent& PtrEvent = static_cast<const FPointerEvent&>(Event);
+		if (PtrEvent.GetEffectingButton() == EKeys::LeftMouseButton)
+		{
+			LastMouseClickTime = FSlateApplication::Get().GetCurrentTime();
+			CursorPingPosition = PtrEvent.GetScreenSpacePosition();
+		}
+	}
+
 	#if WITH_EVENT_LOGGING
 		if (LoggedEvents.Num() >= MaxLoggedEvents)
 		{
@@ -346,6 +384,43 @@ int32 SWidgetReflector::Visualize( const FWidgetPath& InWidgetsToVisualize, FSla
 	return LayerId;
 }
 
+int32 SWidgetReflector::VisualizeCursorAndKeys(FSlateWindowElementList& OutDrawElements, int32 LayerId) const
+{
+	if (bEnableDemoMode)
+	{
+		static const float ClickFadeTime = 0.5f;
+		static const float PingScaleAmount = 3.0f;
+		static const FName CursorPingBrush("DemoRecording.CursorPing");
+		const TSharedPtr<SWindow> WindowBeingDrawn = OutDrawElements.GetWindow();
+
+		// Normalized animation value for the cursor ping between 0 and 1.
+		const float AnimAmount = (FSlateApplication::Get().GetCurrentTime() - LastMouseClickTime) / ClickFadeTime;
+
+		if (WindowBeingDrawn.IsValid() && AnimAmount <= 1.0f)
+		{
+			const FVector2D CursorPosDesktopSpace = CursorPingPosition;
+			const FVector2D CursorSize = FSlateApplication::Get().GetCursorSize();
+			const FVector2D PingSize = CursorSize*PingScaleAmount*FCurveHandle::ApplyEasing(AnimAmount, ECurveEaseFunction::QuadOut);
+			const FLinearColor PingColor = FLinearColor(1.0f, 0.0f, 1.0f, 1.0f - FCurveHandle::ApplyEasing(AnimAmount, ECurveEaseFunction::QuadIn));
+
+			FGeometry CursorHighlightGeometry = FGeometry::MakeRoot(PingSize, FSlateLayoutTransform(CursorPosDesktopSpace - PingSize / 2));
+			CursorHighlightGeometry.AppendTransform(Inverse(WindowBeingDrawn->GetLocalToScreenTransform()));
+
+			FSlateDrawElement::MakeBox(
+				OutDrawElements,
+				LayerId++,
+				CursorHighlightGeometry.ToPaintGeometry(),
+				FCoreStyle::Get().GetBrush(CursorPingBrush),
+				OutDrawElements.GetWindow()->GetClippingRectangleInWindow(),
+				ESlateDrawEffect::None,
+				PingColor
+				);
+		}
+	}
+	
+	return LayerId;
+}
+
 
 /* SWidgetReflector implementation
  *****************************************************************************/
diff --git a/Engine/Source/Developer/SlateReflector/Private/Widgets/SWidgetReflector.h b/Engine/Source/Developer/SlateReflector/Private/Widgets/SWidgetReflector.h
--- a/Engine/Source/Developer/SlateReflector/Private/Widgets/SWidgetReflector.h
+++ b/Engine/Source/Developer/SlateReflector/Private/Widgets/SWidgetReflector.h
@@ -77,6 +77,7 @@ public:
 
 	virtual void SetWidgetsToVisualize( const FWidgetPath& InWidgetsToVisualize ) override;
 	virtual int32 Visualize( const FWidgetPath& InWidgetsToVisualize, FSlateWindowElementList& OutDrawElements, int32 LayerId ) override;
+	virtual int32 VisualizeCursorAndKeys(FSlateWindowElementList& OutDrawElements, int32 LayerId) const override;
 
 protected:
 
@@ -216,4 +217,10 @@ private:
 	void UpdateStats();
 	TSharedRef<ITableRow> GenerateStatRow( TSharedRef<FStatItem> StatItem, const TSharedRef<STableViewBase>& OwnerTable );
 	FReply CopyStatsToClipboard();
+
+private:
+	// DEMO MODE
+	bool bEnableDemoMode;
+	double LastMouseClickTime;
+	FVector2D CursorPingPosition;
 };
diff --git a/Engine/Source/Runtime/Slate/Private/Framework/Application/SlateApplication.cpp b/Engine/Source/Runtime/Slate/Private/Framework/Application/SlateApplication.cpp
--- a/Engine/Source/Runtime/Slate/Private/Framework/Application/SlateApplication.cpp
+++ b/Engine/Source/Runtime/Slate/Private/Framework/Application/SlateApplication.cpp
@@ -983,6 +983,13 @@ void FSlateApplication::DrawWindowAndChildren( const TSharedRef<SWindow>& Window
 			MaxLayerId = WidgetReflector->Visualize( DrawWindowArgs.WidgetsUnderCursor, WindowElementList, MaxLayerId );
 		}
 
+		// Visualize pointer presses and pressed keys for demo-recording purposes.
+		const bool bVisualiseMouseClicks = WidgetReflector.IsValid() && PlatformApplication->Cursor.IsValid() && PlatformApplication->Cursor->GetType() != EMouseCursor::None;
+		if (bVisualiseMouseClicks )
+		{
+			MaxLayerId = WidgetReflector->VisualizeCursorAndKeys( WindowElementList, MaxLayerId );
+		}
+
 		// Keep track of windows that we're actually going to be presenting, so we can mark
 		// them as 'drawn' afterwards.
 		FGenericWindow* NativeWindow = WindowToDraw->GetNativeWindow().Get();
diff --git a/Engine/Source/Runtime/Slate/Public/Framework/Application/IWidgetReflector.h b/Engine/Source/Runtime/Slate/Public/Framework/Application/IWidgetReflector.h
--- a/Engine/Source/Runtime/Slate/Public/Framework/Application/IWidgetReflector.h
+++ b/Engine/Source/Runtime/Slate/Public/Framework/Application/IWidgetReflector.h
@@ -56,6 +56,9 @@ public:
 	 */
 	virtual int32 Visualize( const FWidgetPath& InWidgetsToVisualize, FSlateWindowElementList& OutDrawElements, int32 LayerId ) = 0;
 
+	/** Visualize the cursor position and any pressed keys for demo-recording purposes. */
+	virtual int32 VisualizeCursorAndKeys(FSlateWindowElementList& OutDrawElements, int32 LayerId) const = 0;
+
 	/**
 	 * Sets the widget that should be visualized.
 	 *
diff --git a/Engine/Source/Runtime/SlateCore/Private/Input/Events.cpp b/Engine/Source/Runtime/SlateCore/Private/Input/Events.cpp
--- a/Engine/Source/Runtime/SlateCore/Private/Input/Events.cpp
+++ b/Engine/Source/Runtime/SlateCore/Private/Input/Events.cpp
@@ -14,6 +14,11 @@ FText FInputEvent::ToText() const
 	return NSLOCTEXT("Events", "Unimplemented", "Unimplemented");
 }
 
+bool FInputEvent::IsPointerEvent() const
+{
+	return false;
+}
+
 FText FCharacterEvent::ToText() const
 {
 	return FText::Format( NSLOCTEXT("Events","Char","Char({0})"), FText::FromString(FString(1, &Character)) );
@@ -39,3 +44,8 @@ FText FPointerEvent::ToText() const
 	return FText::Format( NSLOCTEXT("Events","Pointer","Pointer({0}, {1})"), EffectingButton.GetDisplayName() );
 }
 
+bool FPointerEvent::IsPointerEvent() const
+{
+	return true;
+}
+
diff --git a/Engine/Source/Runtime/SlateCore/Private/Styling/CoreStyle.cpp b/Engine/Source/Runtime/SlateCore/Private/Styling/CoreStyle.cpp
--- a/Engine/Source/Runtime/SlateCore/Private/Styling/CoreStyle.cpp
+++ b/Engine/Source/Runtime/SlateCore/Private/Styling/CoreStyle.cpp
@@ -229,6 +229,12 @@ TSharedRef<ISlateStyle> FCoreStyle::Create( const FName& InStyleSetName )
 		Style->Set("NoBorder", NoBorder);
 	}
 
+	
+	// Demo Recording
+	{
+		Style->Set("DemoRecording.CursorPing", new IMAGE_BRUSH("Common/CursorPing", FVector2D(31,31)));
+	}
+
 	// Error Reporting
 	{
 		Style->Set("ErrorReporting.Box", new BOX_BRUSH( "Common/TextBlockHighlightShape", FMargin(3.f / 8.f)));
diff --git a/Engine/Source/Runtime/SlateCore/Public/Input/Events.h b/Engine/Source/Runtime/SlateCore/Public/Input/Events.h
--- a/Engine/Source/Runtime/SlateCore/Public/Input/Events.h
+++ b/Engine/Source/Runtime/SlateCore/Public/Input/Events.h
@@ -325,6 +325,9 @@ public:
 	}
 
 	SLATECORE_API virtual FText ToText() const;
+	
+	/** Is this event a pointer event (touch or cursor). */
+	SLATECORE_API virtual bool IsPointerEvent() const;
 
 protected:
 
@@ -414,7 +417,7 @@ public:
 		return KeyCode;
 	}
 
-	SLATECORE_API virtual FText ToText() const override;
+	SLATECORE_API virtual FText ToText() const override;	
 
 private:
 	// Name of the key that was pressed.
@@ -741,6 +744,8 @@ public:
 
 	SLATECORE_API virtual FText ToText() const override;
 
+	SLATECORE_API virtual bool IsPointerEvent() const override;
+
 	template<typename PointerEventType>
 	static PointerEventType MakeTranslatedEvent( const PointerEventType& InPointerEvent, const FVirtualPointerPosition& VirtualPosition )
 	{
