Added custom slate stats, removed all but top level UE4 slate stats. #BUN
* See SlateStats.h for details on the new system, which allows for full hierarchical profiling with limited overhead.
* Moved SWidget::GetVisibility() out of line so we can instrument it without a full recompile.
* Widget debug info now uses FName to store the line number of the file. Keeps it in one place and the new stats system can utilize it as an FName.
* Exposed SWidget CreatedInFile debug info as an FName so new stats system can use it.
* FSlateVertex no longer uses FVector4, which is an aligned struct. Switched to float[4] instead. Goes from 48 to 36 bytes.

diff --git a/Engine/Source/Developer/SlateReflector/Private/Widgets/SWidgetReflector.cpp b/Engine/Source/Developer/SlateReflector/Private/Widgets/SWidgetReflector.cpp
--- a/Engine/Source/Developer/SlateReflector/Private/Widgets/SWidgetReflector.cpp
+++ b/Engine/Source/Developer/SlateReflector/Private/Widgets/SWidgetReflector.cpp
@@ -2,17 +2,17 @@
 
 #include "SlateReflectorPrivatePCH.h"
 #include "ISlateReflectorModule.h"
-
-#if STATS
-#include "StatsData.h"
-#endif
-
+#include "SlateStats.h"
 
 #define LOCTEXT_NAMESPACE "SWidgetReflector"
 #define WITH_EVENT_LOGGING 0
 
+extern SLATECORE_API int32 bFoldTick;
 static const int32 MaxLoggedEvents = 100;
 
+/**
+ * Widget reflector user widget.
+ */
 
 /* Local helpers
  *****************************************************************************/
@@ -37,20 +37,215 @@ struct FLoggedEvent
 	FText HandlerText;
 };
 
+namespace WidgetReflectorImpl
+{
 
-/* SWidgetReflector interface
- *****************************************************************************/
-
-static FReply ToggleSlateStats()
+/**
+ * Widget reflector implementation
+ */
+class SWidgetReflector : public ::SWidgetReflector
 {
-	#if STATS
-	FSelfRegisteringExec::StaticExec(nullptr, TEXT("stat slate"), *GLog);
-	#endif
+	// The reflector uses a tree that observes FReflectorNodes.
+	typedef STreeView<TSharedPtr<FReflectorNode>> SReflectorTree;
+
+private:
+
+	virtual void Construct( const FArguments& InArgs ) override;
+
+	// SCompoundWidget overrides
+	virtual void Tick( const FGeometry& AllottedGeometry, const double InCurrentTime, const float InDeltaTime ) override;
+
+	// IWidgetReflector interface
+
+	virtual void OnEventProcessed( const FInputEvent& Event, const FReplyBase& InReply ) override;
+
+	virtual bool IsInPickingMode() const override
+	{
+		return bIsPicking;
+	}
+
+	virtual bool IsShowingFocus() const override
+	{
+		return bShowFocus;
+	}
+
+	virtual bool IsVisualizingLayoutUnderCursor() const override
+	{
+		return bIsPicking;
+	}
+
+	virtual void OnWidgetPicked() override
+	{
+		bIsPicking = false;
+	}
+
+	virtual bool ReflectorNeedsToDrawIn( TSharedRef<SWindow> ThisWindow ) const override;
+
+	virtual void SetSourceAccessDelegate( FAccessSourceCode InDelegate ) override
+	{
+		SourceAccessDelegate = InDelegate;
+	}
+
+	virtual void SetAssetAccessDelegate(FAccessAsset InDelegate) override
+	{
+		AsseetAccessDelegate = InDelegate;
+	}
+
+	virtual void SetWidgetsToVisualize( const FWidgetPath& InWidgetsToVisualize ) override;
+	virtual int32 Visualize( const FWidgetPath& InWidgetsToVisualize, FSlateWindowElementList& OutDrawElements, int32 LayerId ) override;
+	virtual int32 VisualizeCursorAndKeys(FSlateWindowElementList& OutDrawElements, int32 LayerId) const override;
+
+	/**
+	 * Generates a tool tip for the given reflector tree node.
+	 *
+	 * @param InReflectorNode The node to generate the tool tip for.
+	 * @return The tool tip widget.
+	 */
+	TSharedRef<SToolTip> GenerateToolTipForReflectorNode( TSharedPtr<FReflectorNode> InReflectorNode );
+
+	/**
+	 * Mark the provided reflector nodes such that they stand out in the tree and are visible.
+	 *
+	 * @param WidgetPathToObserve The nodes to mark.
+	 */
+	void VisualizeAsTree( const TArray< TSharedPtr<FReflectorNode> >& WidgetPathToVisualize );
+
+	/**
+	 * Draw the widget path to the picked widget as the widgets' outlines.
+	 *
+	 * @param InWidgetsToVisualize A widget path whose widgets' outlines to draw.
+	 * @param OutDrawElements A list of draw elements; we will add the output outlines into it.
+	 * @param LayerId The maximum layer achieved in OutDrawElements so far.
+	 * @return The maximum layer ID we achieved while painting.
+	 */
+	int32 VisualizePickAsRectangles( const FWidgetPath& InWidgetsToVisualize, FSlateWindowElementList& OutDrawElements, int32 LayerId );
+
+	/**
+	 * Draw an outline for the specified nodes.
+	 *
+	 * @param InNodesToDraw A widget path whose widgets' outlines to draw.
+	 * @param WindowGeometry The geometry of the window in which to draw.
+	 * @param OutDrawElements A list of draw elements; we will add the output outlines into it.
+	 * @param LayerId the maximum layer achieved in OutDrawElements so far.
+	 * @return The maximum layer ID we achieved while painting.
+	 */
+	int32 VisualizeSelectedNodesAsRectangles( const TArray<TSharedPtr<FReflectorNode>>& InNodesToDraw, const TSharedRef<SWindow>& VisualizeInWindow, FSlateWindowElementList& OutDrawElements, int32 LayerId );
+
+	/** Callback for changing the application scale slider. */
+	void HandleAppScaleSliderChanged( float NewValue )
+	{
+		FSlateApplication::Get().SetApplicationScale(NewValue);
+	}
+
+	FReply HandleDisplayTextureAtlases();
+	FReply HandleDisplayFontAtlases();
+
+	/** Callback for getting the value of the application scale slider. */
+	float HandleAppScaleSliderValue() const
+	{
+		return FSlateApplication::Get().GetApplicationScale();
+	}
+
+	/** Callback for checked state changes of the focus check box. */
+	void HandleFocusCheckBoxCheckedStateChanged( ECheckBoxState NewValue );
+
+	/** Callback for getting the checked state of the focus check box. */
+	ECheckBoxState HandleFocusCheckBoxIsChecked() const
+	{
+		return bShowFocus ? ECheckBoxState::Checked : ECheckBoxState::Unchecked;
+	}
+
+	/** Callback for getting the text of the frame rate text block. */
+	FString HandleFrameRateText() const;
+
+	/** Callback for clicking the pick button. */
+	FReply HandlePickButtonClicked()
+	{
+		bIsPicking = !bIsPicking;
+
+		if (bIsPicking)
+		{
+			bShowFocus = false;
+		}
+
+		return FReply::Handled();
+	}
+
+	/** Callback for getting the color of the pick button text. */
+	FSlateColor HandlePickButtonColorAndOpacity() const
+	{
+		static const FName SelectionColor("SelectionColor");
+		static const FName DefaultForeground("DefaultForeground");
+
+		return bIsPicking
+			? FCoreStyle::Get().GetSlateColor(SelectionColor)
+			: FCoreStyle::Get().GetSlateColor(DefaultForeground);
+	}
+
+	/** Callback for getting the text of the pick button. */
+	FString HandlePickButtonText() const;
+
+	/** Callback for generating a row in the reflector tree view. */
+	TSharedRef<ITableRow> HandleReflectorTreeGenerateRow( TSharedPtr<FReflectorNode> InReflectorNode, const TSharedRef<STableViewBase>& OwnerTable );
+
+	/** Callback for getting the child items of the given reflector tree node. */
+	void HandleReflectorTreeGetChildren( TSharedPtr<FReflectorNode> InWidgetGeometry, TArray<TSharedPtr<FReflectorNode>>& OutChildren );
+
+	/** Callback for when the selection in the reflector tree has changed. */
+	void HandleReflectorTreeSelectionChanged( TSharedPtr< FReflectorNode >, ESelectInfo::Type /*SelectInfo*/ )
+	{
+		SelectedNodes = ReflectorTree->GetSelectedItems();
+	}
+
+	TSharedRef<ITableRow> GenerateEventLogRow( TSharedRef<FLoggedEvent> InReflectorNode, const TSharedRef<STableViewBase>& OwnerTable );
+
+	TArray< TSharedRef<FLoggedEvent> > LoggedEvents;
+	TSharedPtr< SListView< TSharedRef< FLoggedEvent > > > EventListView;
+	TSharedPtr< SHorizontalBox > StatsToolsBox;
+	TSharedPtr<SReflectorTree> ReflectorTree;
+
+	TArray<TSharedPtr<FReflectorNode>> SelectedNodes;
+	TArray<TSharedPtr<FReflectorNode>> ReflectorTreeRoot;
+	TArray<TSharedPtr<FReflectorNode>> PickedPath;
+
+	SSplitter::FSlot* WidgetInfoLocation;
+
+	FAccessSourceCode SourceAccessDelegate;
+	FAccessAsset AsseetAccessDelegate;
+
+	bool bShowFocus;
+	bool bIsPicking;
+
+#if SLATE_STATS
+	// STATS
+	TSharedPtr<SBorder> StatsBorder;
+	TArray< TSharedRef<class FStatItem> > StatsItems;
+	TSharedPtr< SListView< TSharedRef<FStatItem> > > StatsList;
+
+	TSharedRef<SWidget> MakeStatViewer();
+	void UpdateStats();
+	TSharedRef<ITableRow> GenerateStatRow(TSharedRef<FStatItem> StatItem, const TSharedRef<STableViewBase>& OwnerTable);
+#endif
+private:
+	// DEMO MODE
+	bool bEnableDemoMode;
+	double LastMouseClickTime;
+	FVector2D CursorPingPosition;
+};
+
+
+
+
+
+
+
+
+
+
+
+
 
-	return FReply::Handled();
-}
 
-extern SLATECORE_API int32 bFoldTick;
 
 
 void SWidgetReflector::Construct( const FArguments& InArgs )
@@ -77,12 +272,12 @@ void SWidgetReflector::Construct( const FArguments& InArgs )
 						SNew(SHorizontalBox)
 					
 						+ SHorizontalBox::Slot()
-						.AutoWidth()
-						.VAlign(VAlign_Center)
-						[
-							SNew(STextBlock)
-								.Text(LOCTEXT("AppScale", "Application Scale: ").ToString())
-						]
+							.AutoWidth()
+							.VAlign(VAlign_Center)
+							[
+								SNew(STextBlock)
+									.Text(LOCTEXT("AppScale", "Application Scale: ").ToString())
+							]
 
 						+ SHorizontalBox::Slot()
 						.AutoWidth()
@@ -99,13 +294,11 @@ void SWidgetReflector::Construct( const FArguments& InArgs )
 									.OnValueChanged(this, &SWidgetReflector::HandleAppScaleSliderChanged)
 							]
 						]
-
 						+ SHorizontalBox::Slot()
 						.FillWidth(1.0f)
 						[
 							SNew(SSpacer)
 						]
-
 						+ SHorizontalBox::Slot()
 						.AutoWidth()
 						.Padding(FMargin(5.0f, 0.0f))
@@ -136,12 +329,12 @@ void SWidgetReflector::Construct( const FArguments& InArgs )
 					[
 						// Check box that controls LIVE MODE
 						SNew(SCheckBox)
-						.IsChecked(this, &SWidgetReflector::HandleFocusCheckBoxIsChecked)
-						.OnCheckStateChanged(this, &SWidgetReflector::HandleFocusCheckBoxCheckedStateChanged)
-						[
-							SNew(STextBlock)
-							.Text(LOCTEXT("ShowFocus", "Show Focus").ToString())
-						]
+							.IsChecked(this, &SWidgetReflector::HandleFocusCheckBoxIsChecked)
+							.OnCheckStateChanged(this, &SWidgetReflector::HandleFocusCheckBoxCheckedStateChanged)
+							[
+								SNew(STextBlock)
+									.Text(LOCTEXT("ShowFocus", "Show Focus").ToString())
+							]
 					]
 
 					+ SHorizontalBox::Slot()
@@ -150,47 +343,14 @@ void SWidgetReflector::Construct( const FArguments& InArgs )
 					[
 						// Check box that controls PICKING A WIDGET TO INSPECT
 						SNew(SButton)
-						.OnClicked(this, &SWidgetReflector::HandlePickButtonClicked)
-						.ButtonColorAndOpacity(this, &SWidgetReflector::HandlePickButtonColorAndOpacity)
-						[
-							SNew(STextBlock)
-							.Text(this, &SWidgetReflector::HandlePickButtonText)
-						]
-					]
-
-					+ SHorizontalBox::Slot()
-					.AutoWidth()
-					.Padding(5.0f)
-					[
-						SNew(SButton)
-						#if STATS
-						.Text( LOCTEXT("ToggleStats", "Toggle Stats") )
-						#else
-						.Text( LOCTEXT("ToggleStatsUnavailable", "Toggle Stats [see tooltip]") )
-						.ToolTip
-						(
-							SNew(SToolTip)
+							.OnClicked(this, &SWidgetReflector::HandlePickButtonClicked)
+							.ButtonColorAndOpacity(this, &SWidgetReflector::HandlePickButtonColorAndOpacity)
 							[
 								SNew(STextBlock)
-								.WrapTextAt(200.0f)
-								.Text( LOCTEXT("ToggleStatsUnavailableTooltip", "To enable STATS in standalone applications turn on <bCompileWithStatsWithoutEngine>false</bCompileWithStatsWithoutEngine> in BuildConfiguration.xml") )
+									.Text(this, &SWidgetReflector::HandlePickButtonText)
 							]
-						)
-						.IsEnabled(false)
-						#endif
-						.OnClicked_Static( &ToggleSlateStats )
-					]
-
-					+ SHorizontalBox::Slot()
-					.AutoWidth()
-					.Padding(5.0f)
-					[
-						SNew(SButton)
-						.OnClicked( this, &SWidgetReflector::CopyStatsToClipboard )
-						.Text( LOCTEXT("CopyStatsToClipboard", "Copy Stats") )
 					]
-
-					+ SHorizontalBox::Slot()
+					+SHorizontalBox::Slot()
 					.AutoWidth()
 					.Padding(5.0f)
 					[
@@ -238,8 +398,182 @@ void SWidgetReflector::Construct( const FArguments& InArgs )
 							]
 						]
 					]
+#if SLATE_STATS
+					+SHorizontalBox::Slot()
+					.AutoWidth()
+					.Padding(5.0f)
+					[
+						SNew(SCheckBox)
+						.Style( FCoreStyle::Get(), "ToggleButtonCheckbox" )
+						.IsChecked_Static([]
+						{
+							extern SLATECORE_API int32 GSlateStatsFlatEnable;
+							return GSlateStatsFlatEnable == 0 ? ECheckBoxState::Unchecked : ECheckBoxState::Checked;
+						})
+						.OnCheckStateChanged_Lambda( [=]( const ECheckBoxState NewState )
+						{
+							extern SLATECORE_API int32 GSlateStatsFlatEnable;
+							GSlateStatsFlatEnable = (NewState == ECheckBoxState::Checked) ? 1 : 0;
+							StatsToolsBox->SetVisibility(GSlateStatsFlatEnable != 0 ? EVisibility::Visible : EVisibility::Collapsed);
+						})
+						.ToolTip
+						(
+							SNew(SToolTip)
+							[
+								SNew(STextBlock)
+								.WrapTextAt(200.0f)
+								.Text( LOCTEXT("ToggleStatsTooltip", "Enables flat stats view.") )
+							]
+						)
+						[
+							SNew(SBox)
+							.VAlign( VAlign_Center )
+							.HAlign( HAlign_Center )
+							[
+								SNew(STextBlock)	
+								.Text( LOCTEXT("ToggleStats", "Toggle Stats") )
+							]
+						]
+					]
+#else
+					+SHorizontalBox::Slot()
+					.AutoWidth()
+					.Padding(5.0f)
+					[
+						SNew(SCheckBox)
+						.Style( FCoreStyle::Get(), "ToggleButtonCheckbox" )
+						.ToolTip
+						(
+							SNew(SToolTip)
+							[
+								SNew(STextBlock)
+								.WrapTextAt(200.0f)
+								.Text( LOCTEXT("ToggleStatsUnavailableTooltip", "To enable slate stats, compile with SLATE_STATS defined to one (see SlateStats.h).") )
+							]
+						)
+						[
+							SNew(SBox)
+							.VAlign( VAlign_Center )
+							.HAlign( HAlign_Center )
+							[
+								SNew(STextBlock)	
+								.Text( LOCTEXT("ToggleStatsUnavailable", "Toggle Stats [see toolip]") )
+							]
+						]
+						.IsEnabled(false)
+					]
+#endif
 				]
-
+#if SLATE_STATS
+				+ SVerticalBox::Slot()
+				.AutoHeight()
+				[
+					SAssignNew(StatsToolsBox, SHorizontalBox)
+					.Visibility(EVisibility::Collapsed)
+					+SHorizontalBox::Slot()
+					.AutoWidth()
+					.Padding(5.0f)
+					[
+						SNew(SCheckBox)
+						.Style( FCoreStyle::Get(), "ToggleButtonCheckbox" )
+						.IsChecked_Static([]
+						{
+							extern SLATECORE_API int32 GSlateStatsFlatLogOutput;
+							return GSlateStatsFlatLogOutput == 0 ? ECheckBoxState::Unchecked : ECheckBoxState::Checked;
+						})
+						.OnCheckStateChanged_Static([]( const ECheckBoxState NewState )
+						{
+							extern SLATECORE_API int32 GSlateStatsFlatLogOutput;
+							GSlateStatsFlatLogOutput = (NewState == ECheckBoxState::Checked) ? 1 : 0;
+						})
+						.ToolTip
+						(
+							SNew(SToolTip)
+							[
+								SNew(STextBlock)
+								.WrapTextAt(200.0f)
+								.Text( LOCTEXT("LogStatsTooltip", "Enables outputting stats to the log at the given interval.") )
+							]
+						)
+						[
+							SNew(SBox)
+							.VAlign( VAlign_Center )
+							.HAlign( HAlign_Center )
+							[
+								SNew(STextBlock)	
+								.Text( LOCTEXT("ToggleLogStats", "Log Stats") )
+							]
+						]
+					]
+					+SHorizontalBox::Slot()
+					.AutoWidth()
+					.Padding(5.0f)
+					[
+						SNew(SButton)
+						.OnClicked_Static([]
+						{
+							extern SLATECORE_API int32 GSlateStatsHierarchyTrigger;
+							GSlateStatsHierarchyTrigger = 1;
+							return FReply::Handled();
+						})
+						.ToolTip
+						(
+							SNew(SToolTip)
+							[
+								SNew(STextBlock)
+								.WrapTextAt(200.0f)
+								.Text( LOCTEXT("CaptureStatsHierarchyTooltip", "When clicked, the next rendered frame will capture hierarchical stats and save them to file in the Saved/ folder with the following name: SlateHierarchyStats-<timestamp>.csv") )
+							]
+						)
+						[
+							SNew(STextBlock)
+								.Text(LOCTEXT("CaptureHierarchy", "Capture Hierarchy"))
+						]
+					]
+					+ SHorizontalBox::Slot()
+					.AutoWidth()
+					.VAlign(VAlign_Center)
+					[
+						SNew(STextBlock)
+						.ToolTip
+						(
+							SNew(SToolTip)
+							[
+								SNew(STextBlock)
+								.WrapTextAt(200.0f)
+								.Text( LOCTEXT("StatsSamplingIntervalTooltip", "the interval (in seconds) to integrate stats before updating the averages.") )
+							]
+						)
+						.Text(LOCTEXT("StatsSampleWindow", "Sampling Interval: ").ToString())
+					]
+					+ SHorizontalBox::Slot()
+					[
+						SNew(SSpinBox<float>)
+						.ToolTip
+						(
+							SNew(SToolTip)
+							[
+								SNew(STextBlock)
+								.WrapTextAt(200.0f)
+								.Text( LOCTEXT("StatsSamplingIntervalTooltip", "the interval (in seconds) to integrate stats before updating the stats.") )
+							]
+						)
+						.Value_Static([]
+						{
+							extern SLATECORE_API float GSlateStatsFlatIntervalWindowSec;
+							return GSlateStatsFlatIntervalWindowSec;
+						})
+						.MinValue(0.1f)
+						.MaxValue(15.0f)
+						.Delta(0.1f)
+						.OnValueChanged_Static([](float NewValue)
+						{
+							extern SLATECORE_API float GSlateStatsFlatIntervalWindowSec;
+							GSlateStatsFlatIntervalWindowSec = NewValue;
+						})
+					]
+				]
+#endif
 				+ SVerticalBox::Slot()
 					.FillHeight(1.0f)
 					[
@@ -255,7 +589,7 @@ void SWidgetReflector::Construct( const FArguments& InArgs )
 							SNew(SHeaderRow)
 
 							+ SHeaderRow::Column("WidgetName")
-							.DefaultLabel(LOCTEXT("WidgetName", "Widget Name"))
+							.DefaultLabel(LOCTEXT("WidgetName", "Widget Name").ToString())
 							.FillWidth(0.65f)
 
 							+ SHeaderRow::Column("ForegroundColor")
@@ -264,16 +598,16 @@ void SWidgetReflector::Construct( const FArguments& InArgs )
 							.HeaderContent()
 							[
 								SNew(STextBlock)
-								.Text(LOCTEXT("ForegroundColor", "FG"))
-								.ToolTipText(LOCTEXT("ForegroundColorToolTip", "Foreground Color"))
+								.Text(LOCTEXT("ForegroundColor", "FG").ToString())
+								.ToolTipText(LOCTEXT("ForegroundColorToolTip", "Foreground Color").ToString())
 							]
 
 							+ SHeaderRow::Column("Visibility")
-							.DefaultLabel(LOCTEXT("Visibility", "Visibility" ))
+							.DefaultLabel(LOCTEXT("Visibility", "Visibility" ).ToString())
 							.FixedWidth(125.0f)
 
 							+ SHeaderRow::Column("WidgetInfo")
-							.DefaultLabel(LOCTEXT("WidgetInfo", "Widget Info" ))
+							.DefaultLabel(LOCTEXT("WidgetInfo", "Widget Info" ).ToString())
 							.FillWidth(0.25f)
 
 							+ SHeaderRow::Column("Address")
@@ -291,15 +625,21 @@ void SWidgetReflector::Construct( const FArguments& InArgs )
 							.OnGenerateRow(this, &SWidgetReflector::GenerateEventLogRow)
 					]
 				#endif //WITH_EVENT_LOGGING
-			
+#if SLATE_STATS			
 				+ SVerticalBox::Slot()
 					.AutoHeight()
 					[
-					MakeStatViewer()
+						MakeStatViewer()
 					]
-				
+#endif
 			]
 	];
+
+	if (StatsToolsBox.IsValid())
+	{
+		extern SLATECORE_API int32 GSlateStatsFlatEnable;
+		StatsToolsBox->SetVisibility(GSlateStatsFlatEnable != 0 ? EVisibility::Visible : EVisibility::Collapsed);
+	}
 }
 
 /* SCompoundWidget overrides
@@ -307,7 +647,9 @@ void SWidgetReflector::Construct( const FArguments& InArgs )
 
 void SWidgetReflector::Tick( const FGeometry& AllottedGeometry, const double InCurrentTime, const float InDeltaTime )
 {
+#if SLATE_STATS
 	UpdateStats();
+#endif
 }
 
 
@@ -594,52 +936,36 @@ TSharedRef<ITableRow> SWidgetReflector::GenerateEventLogRow( TSharedRef<FLoggedE
 	];
 }
 
+#if SLATE_STATS
 
 //
 // STATS
 //
 
-/**
- * Holds the stat data for visualation. StatItems have an immutable
- * identity and a mutable value that changes from frame to frame.
- */
 class FStatItem
 {
 public:
 	
-	FStatItem( FName InName, FText InDisplayText )
-	: StatName( InName )
-	, DisplayText( InDisplayText )
-	, InclusiveAvgMs( 0.0f )
+	FStatItem(FSlateStatCycleCounter* InCounter)
+	: Counter(InCounter)
 	{
+		UpdateValues();
 	}
 
-	FName GetStatName() const { return StatName; }
-	FText GetDisplayName() const { return DisplayText; }	
-	FText GetInclusiveAvgMsText() const
-	{
-		return InclusiveAvgMsText;
-	}
+	FText GetStatName() const { return StatName; }
+	FText GetInclusiveAvgMsText() const { return InclusiveAvgMsText; }
 	float GetInclusiveAvgMs() const { return InclusiveAvgMs; }
-	void UpdateInclusiveAvgMs( float InMs )
+	void UpdateValues()
 	{
-		static const auto FormattingOptions = FNumberFormattingOptions()
-			.SetMinimumIntegralDigits(1)
-			.SetMinimumFractionalDigits(3)
-			.SetMaximumFractionalDigits(3);
-
-		InclusiveAvgMs = InMs;
-		InclusiveAvgMsText = FText::AsNumber(InMs, &FormattingOptions);
+		StatName = FText::FromName(Counter->GetName());
+		InclusiveAvgMsText = FText::AsNumber(Counter->GetLastComputedAverageInclusiveTime(), &FNumberFormattingOptions().SetMinimumIntegralDigits(1).SetMinimumFractionalDigits(3).SetMaximumFractionalDigits(3));
+		InclusiveAvgMs = (float)(Counter->GetLastComputedAverageInclusiveTime());
 	}
-
 private:
-
-	const FName StatName;
-	const FText DisplayText;
-
+	FSlateStatCycleCounter* Counter;
+	FText StatName;
 	FText InclusiveAvgMsText;
 	float InclusiveAvgMs;
-	
 };
 
 static const FName ColumnId_StatName("StatName");
@@ -648,12 +974,24 @@ static const FName ColumnId_InclusiveAvgMsGraph("InclusiveAvgMsGraph");
 
 TSharedRef<SWidget> SWidgetReflector::MakeStatViewer()
 {
+	// the list of registered counters must remain constant throughout program execution.
+	// As long as all counters are declared globally this will be true.
+	for (auto Stat : FSlateStatCycleCounter::GetRegisteredCounters())
+	{
+		StatsItems.Add(MakeShareable(new FStatItem(Stat)));
+	}
+
 	return
 		SAssignNew(StatsBorder, SBorder)
 		.BorderImage( FCoreStyle::Get().GetBrush("NoBrush") )
+		.Visibility_Lambda([]
+		{
+			extern SLATECORE_API int32 GSlateStatsFlatEnable;
+			return GSlateStatsFlatEnable > 0 ? EVisibility::Visible : EVisibility::Collapsed;
+		})
 		[
 			// STATS LIST
-			SAssignNew( StatsList, SListView< TSharedRef< FStatItem > > )
+			SAssignNew( StatsList, SListView< TSharedRef<FStatItem> > )
 			.OnGenerateRow( this, &SWidgetReflector::GenerateStatRow )
 			.ListItemsSource( &StatsItems )
 			.HeaderRow
@@ -680,188 +1018,98 @@ TSharedRef<SWidget> SWidgetReflector::MakeStatViewer()
 
 void SWidgetReflector::UpdateStats()
 {
-	#if STATS
-
-	const FGameThreadHudData* ViewDataPtr = FHUDGroupGameThreadRenderer::Get().Latest;
-	if (ViewDataPtr != nullptr)
+	if (FSlateStatCycleCounter::AverageInclusiveTimesWereUpdatedThisFrame())
 	{
-		// Stats are up to date; make sure the UI is full bright.
-		StatsBorder->SetColorAndOpacity(FLinearColor::White);
-
-		// If we don't get a report about a stat this frame, assume its value is 0;
-		for (auto& StatItem : StatsItems)
-		{
-			StatItem->UpdateInclusiveAvgMs(0.0f);
-		}
-
-		// Update stats with latest values; add any stats if they are missing
-		bool bNeedsRefresh = false;
-		const FGameThreadHudData &ViewData = *ViewDataPtr;
-		for (int32 GroupIndex = 0; GroupIndex < ViewData.HudGroups.Num(); ++GroupIndex)
+		for (auto& StatsItem : StatsItems)
 		{
-			const FHudGroup& HudGroup = ViewData.HudGroups[GroupIndex];
-			const bool bHasHierarchy = !!HudGroup.HierAggregate.Num();
-			const bool bHasFlat = !!HudGroup.FlatAggregate.Num();
-
-			int32 Row = 0;
-			for (const FComplexStatMessage& StatItem : HudGroup.FlatAggregate)
-			{
-				const TSharedRef<FStatItem>* ExistingItemPtr = StatsItems.FindByPredicate( [&StatItem]( const TSharedRef<FStatItem>& ItemRef )
-				{
-					return ItemRef->GetStatName() == StatItem.GetShortName();
-				});
-
-				const float InMs = FPlatformTime::ToMilliseconds(StatItem.GetValue_Duration(EComplexStatField::IncAve));
-
-				// If stat does not exist...
-				if (ExistingItemPtr == nullptr)
-				{
-					// ... add it.
-					const TSharedRef<FStatItem> NewStatItem( new FStatItem( StatItem.GetShortName(), FText::FromString(StatItem.GetDescription() ) ) );
-					StatsItems.Add( NewStatItem );
-					ExistingItemPtr = &StatsItems.Last();
-
-					// We have added items to the list, and it needs to be refreshed.
-					bNeedsRefresh = true;
-				}
-
-				// Update stat values.
-				(*ExistingItemPtr)->UpdateInclusiveAvgMs( InMs );
-			}
-		}
-
-		// If we add widgets in tick, we must refresh the list.
-		if (bNeedsRefresh)
-		{
-			StatsList->RequestListRefresh();
-		}		
-		else
-		{
-			int32 a = 5;
-			if (a == 10){  }
+			StatsItem->UpdateValues();
 		}
+		//StatsList->RequestListRefresh();
 	}
-	else
-	{
-		// We are now showing stale stats; provide visual hint.
-		StatsBorder->SetColorAndOpacity(FLinearColor(1,1,1,0.75f));
-		
-		// If we don't get a report about a stat this frame, assume its value is 0;
-		for (auto& StatItem : StatsItems)
-		{
-			StatItem->UpdateInclusiveAvgMs(0.0f);
-		}
-	}
-
-	#endif
 }
 
-TSharedRef<ITableRow> SWidgetReflector::GenerateStatRow( TSharedRef<FStatItem> StatItem, const TSharedRef<STableViewBase>& OwnerTable )
+
+class SStatTableRow : public SMultiColumnTableRow< TSharedRef<FSlateStatCycleCounter> >
 {
-	class SStatTableRow : public SMultiColumnTableRow< TSharedRef<FStatItem> >
-	{
 	
-	public:
-		typedef SMultiColumnTableRow< TSharedRef<FStatItem> > Super;
+public:
+	typedef SMultiColumnTableRow< TSharedRef<FSlateStatCycleCounter> > Super;
 		
-		SLATE_BEGIN_ARGS( SStatTableRow )
-		{}
-		SLATE_END_ARGS();
-
-		void Construct( const FArguments& InArgs, const TSharedRef<STableViewBase>& OwnerTable, const TSharedRef<FStatItem>& InStatItem )
-		{
-			StatItem = InStatItem;
-			FSuperRowType::Construct( FTableRowArgs(), OwnerTable );
-		}
+	SLATE_BEGIN_ARGS( SStatTableRow )
+	{}
+	SLATE_END_ARGS();
 
-		virtual TSharedRef<SWidget> GenerateWidgetForColumn( const FName& ColumnName ) override
-		{
-			if (ColumnName == ColumnId_StatName)
-			{
-				// STAT NAME
-				return
-				SNew(STextBlock)
-				.Text( StatItem->GetDisplayName() );
-			}
-			else if (ColumnName == ColumnId_InclusiveAvgMs)
-			{
-				// STAT NUMBER
-				return
-				SNew(SBox)
-				.Padding(FMargin(5,0))
-				[
-					SNew(STextBlock)
-					.TextStyle(FCoreStyle::Get(), "MonospacedText")
-					.Text( this, &SStatTableRow::GetInclusiveAvgMsText )
-				];
-			}
-			else if ( ColumnName == ColumnId_InclusiveAvgMsGraph )
-			{
-				// BAR GRAPH
-				return
-				SNew(SHorizontalBox)
-				+SHorizontalBox::Slot()
-				.Padding(0,1)
-				.FillWidth( TAttribute<float>(SharedThis(this), &SStatTableRow::BarGraph_GetFilledWidth) )
-				[
-					SNew(SImage)
-					.Image( FCoreStyle::Get().GetBrush("WhiteBrush") )
-					.ColorAndOpacity( this, &SStatTableRow::BarGraph_GetBarColor )
-				]
-				+SHorizontalBox::Slot()
-				.Padding(0,1)
-				.FillWidth( TAttribute<float>(SharedThis(this), &SStatTableRow::BarGraph_GetEmptyWidth) );
-			}
-			else
-			{
-				return SNew(SSpacer);
-			}
-		}
-
-		private:
-
-		FText GetInclusiveAvgMsText() const
-		{
-			return StatItem->GetInclusiveAvgMsText();
-		}
+	void Construct( const FArguments& InArgs, const TSharedRef<STableViewBase>& OwnerTable, const TSharedPtr<FStatItem>& InStatItem )
+	{
+		StatItem = InStatItem;
+		FSuperRowType::Construct( FTableRowArgs(), OwnerTable );
+	}
 
-		float BarGraph_GetFilledWidth() const
-		{
-			return StatItem->GetInclusiveAvgMs();
-		}
+	float GetValue() const { return 6.0f; }
 
-		float BarGraph_GetEmptyWidth()  const
-		{
-			static const float MaxMs = 60.0f;
-			return MaxMs - StatItem->GetInclusiveAvgMs();
-		}
+	virtual TSharedRef<SWidget> GenerateWidgetForColumn( const FName& ColumnName ) override;
 
-		FSlateColor BarGraph_GetBarColor() const
-		{
-			static const float RedlineMs = 30.0f;
-			return FMath::Lerp( FLinearColor::Green, FLinearColor::Red, StatItem->GetInclusiveAvgMs()/RedlineMs );
-		}
+private:
+	TSharedPtr<FStatItem> StatItem;
+};
 
-		TSharedPtr<FStatItem> StatItem;
-	};
+TSharedRef<SWidget> SStatTableRow::GenerateWidgetForColumn(const FName& ColumnName)
+{
+	if (ColumnName == ColumnId_StatName)
+	{
+		// STAT NAME
+		return
+			SNew(STextBlock)
+			.Text_Lambda([=] { return StatItem->GetStatName(); });
+	}
+	else if (ColumnName == ColumnId_InclusiveAvgMs)
+	{
+		// STAT NUMBER
+		return
+			SNew(SBox)
+			.Padding(FMargin(5, 0))
+			[
+				SNew(STextBlock)
+				.TextStyle(FCoreStyle::Get(), "MonospacedText")
+				.Text_Lambda([=] { return StatItem->GetInclusiveAvgMsText(); })
+			];
+	}
+	else if (ColumnName == ColumnId_InclusiveAvgMsGraph)
+	{
+		// BAR GRAPH
+		return
+			SNew(SHorizontalBox)
+			+ SHorizontalBox::Slot()
+			.Padding(0, 1)
+			.FillWidth(TAttribute<float>::Create(TAttribute<float>::FGetter::CreateLambda([=] { return StatItem->GetInclusiveAvgMs(); })))
+			[
+				SNew(SImage)
+				.Image(FCoreStyle::Get().GetBrush("WhiteBrush"))
+				.ColorAndOpacity_Lambda([=] { return FMath::Lerp(FLinearColor::Green, FLinearColor::Red, StatItem->GetInclusiveAvgMs() / 30.0f); })
+			]
+			+ SHorizontalBox::Slot()
+			.Padding(0, 1)
+			.FillWidth(TAttribute<float>::Create(TAttribute<float>::FGetter::CreateLambda([=] { return 60.0f - StatItem->GetInclusiveAvgMs(); })))
+			;
+	}
+	else
+	{
+		return SNew(SSpacer);
+	}
+}
 
+TSharedRef<ITableRow> SWidgetReflector::GenerateStatRow( TSharedRef<FStatItem> StatItem, const TSharedRef<STableViewBase>& OwnerTable )
+{
 	return SNew(SStatTableRow, OwnerTable, StatItem);
 }
 
-FReply SWidgetReflector::CopyStatsToClipboard()
-{
-	FString ClipboardString;
+#endif // SLATE_STATS
 
-	for (const TSharedRef<FStatItem>& StatItem : StatsItems)
-	{
-		ClipboardString += StatItem->GetDisplayName().ToString();
-		ClipboardString += TEXT(", ");
-		ClipboardString += StatItem->GetInclusiveAvgMsText().ToString();
-		ClipboardString += TEXT("\r\n");
-	}
+} // namespace WidgetReflectorImpl
 
-	FPlatformMisc::ClipboardCopy( *ClipboardString );
-	return FReply::Handled();
+TSharedRef<SWidgetReflector> SWidgetReflector::New()
+{
+  return MakeShareable( new WidgetReflectorImpl::SWidgetReflector() );
 }
 
 #undef LOCTEXT_NAMESPACE
diff --git a/Engine/Source/Developer/SlateReflector/Private/Widgets/SWidgetReflector.h b/Engine/Source/Developer/SlateReflector/Private/Widgets/SWidgetReflector.h
--- a/Engine/Source/Developer/SlateReflector/Private/Widgets/SWidgetReflector.h
+++ b/Engine/Source/Developer/SlateReflector/Private/Widgets/SWidgetReflector.h
@@ -4,223 +4,20 @@
 
 #include "IWidgetReflector.h"
 
-
-// forward declarations
-struct FLoggedEvent;
-
-
 /**
- * Abstract base class for widget reflectors.
+ * Widget reflector implementation.
+ * User widget to enable iteration without recompilation.
  */
 class SLATEREFLECTOR_API SWidgetReflector
-	: public SCompoundWidget
+	: public SUserWidget
 	, public IWidgetReflector
 {
-	// The reflector uses a tree that observes FReflectorNodes.
-	typedef STreeView<TSharedPtr<FReflectorNode>> SReflectorTree;
-
 public:
 
-	SLATE_BEGIN_ARGS(SWidgetReflector) { }
+	SLATE_USER_ARGS(SWidgetReflector) 
+	{ }
 	SLATE_END_ARGS()
 
-	/**
-	 * Creates and initializes a new widget reflector widget.
-	 *
-	 * @param InArgs The construction arguments.
-	 */
-	void Construct( const FArguments& InArgs );
-
-public:
-
-	// SCompoundWidget overrides
-
-	virtual void Tick( const FGeometry& AllottedGeometry, const double InCurrentTime, const float InDeltaTime ) override;
-
-public:
-
-	// IWidgetReflector interface
-
-	virtual void OnEventProcessed( const FInputEvent& Event, const FReplyBase& InReply ) override;
-
-	virtual bool IsInPickingMode() const override
-	{
-		return bIsPicking;
-	}
-
-	virtual bool IsShowingFocus() const override
-	{
-		return bShowFocus;
-	}
-
-	virtual bool IsVisualizingLayoutUnderCursor() const override
-	{
-		return bIsPicking;
-	}
-
-	virtual void OnWidgetPicked() override
-	{
-		bIsPicking = false;
-	}
-
-	virtual bool ReflectorNeedsToDrawIn( TSharedRef<SWindow> ThisWindow ) const override;
-
-	virtual void SetSourceAccessDelegate( FAccessSourceCode InDelegate ) override
-	{
-		SourceAccessDelegate = InDelegate;
-	}
-
-	virtual void SetAssetAccessDelegate(FAccessAsset InDelegate) override
-	{
-		AsseetAccessDelegate = InDelegate;
-	}
-
-	virtual void SetWidgetsToVisualize( const FWidgetPath& InWidgetsToVisualize ) override;
-	virtual int32 Visualize( const FWidgetPath& InWidgetsToVisualize, FSlateWindowElementList& OutDrawElements, int32 LayerId ) override;
-	virtual int32 VisualizeCursorAndKeys(FSlateWindowElementList& OutDrawElements, int32 LayerId) const override;
-
-protected:
-
-	/**
-	 * Generates a tool tip for the given reflector tree node.
-	 *
-	 * @param InReflectorNode The node to generate the tool tip for.
-	 * @return The tool tip widget.
-	 */
-	TSharedRef<SToolTip> GenerateToolTipForReflectorNode( TSharedPtr<FReflectorNode> InReflectorNode );
-
-	/**
-	 * Mark the provided reflector nodes such that they stand out in the tree and are visible.
-	 *
-	 * @param WidgetPathToObserve The nodes to mark.
-	 */
-	void VisualizeAsTree( const TArray< TSharedPtr<FReflectorNode> >& WidgetPathToVisualize );
-
-	/**
-	 * Draw the widget path to the picked widget as the widgets' outlines.
-	 *
-	 * @param InWidgetsToVisualize A widget path whose widgets' outlines to draw.
-	 * @param OutDrawElements A list of draw elements; we will add the output outlines into it.
-	 * @param LayerId The maximum layer achieved in OutDrawElements so far.
-	 * @return The maximum layer ID we achieved while painting.
-	 */
-	int32 VisualizePickAsRectangles( const FWidgetPath& InWidgetsToVisualize, FSlateWindowElementList& OutDrawElements, int32 LayerId );
-
-	/**
-	 * Draw an outline for the specified nodes.
-	 *
-	 * @param InNodesToDraw A widget path whose widgets' outlines to draw.
-	 * @param WindowGeometry The geometry of the window in which to draw.
-	 * @param OutDrawElements A list of draw elements; we will add the output outlines into it.
-	 * @param LayerId the maximum layer achieved in OutDrawElements so far.
-	 * @return The maximum layer ID we achieved while painting.
-	 */
-	int32 VisualizeSelectedNodesAsRectangles( const TArray<TSharedPtr<FReflectorNode>>& InNodesToDraw, const TSharedRef<SWindow>& VisualizeInWindow, FSlateWindowElementList& OutDrawElements, int32 LayerId );
-
-private:
-
-	/** Callback for changing the application scale slider. */
-	void HandleAppScaleSliderChanged( float NewValue )
-	{
-		FSlateApplication::Get().SetApplicationScale(NewValue);
-	}
-
-	FReply HandleDisplayTextureAtlases();
-
-	FReply HandleDisplayFontAtlases();
-
-	/** Callback for getting the value of the application scale slider. */
-	float HandleAppScaleSliderValue() const
-	{
-		return FSlateApplication::Get().GetApplicationScale();
-	}
-
-	/** Callback for checked state changes of the focus check box. */
-	void HandleFocusCheckBoxCheckedStateChanged( ECheckBoxState NewValue );
-
-	/** Callback for getting the checked state of the focus check box. */
-	ECheckBoxState HandleFocusCheckBoxIsChecked() const
-	{
-		return bShowFocus ? ECheckBoxState::Checked : ECheckBoxState::Unchecked;
-	}
-
-	/** Callback for getting the text of the frame rate text block. */
-	FString HandleFrameRateText() const;
-
-	/** Callback for clicking the pick button. */
-	FReply HandlePickButtonClicked()
-	{
-		bIsPicking = !bIsPicking;
-
-		if (bIsPicking)
-		{
-			bShowFocus = false;
-		}
-
-		return FReply::Handled();
-	}
-
-	/** Callback for getting the color of the pick button text. */
-	FSlateColor HandlePickButtonColorAndOpacity() const
-	{
-		static const FName SelectionColor("SelectionColor");
-		static const FName DefaultForeground("DefaultForeground");
-
-		return bIsPicking
-			? FCoreStyle::Get().GetSlateColor(SelectionColor)
-			: FCoreStyle::Get().GetSlateColor(DefaultForeground);
-	}
-
-	/** Callback for getting the text of the pick button. */
-	FString HandlePickButtonText() const;
-
-	/** Callback for generating a row in the reflector tree view. */
-	TSharedRef<ITableRow> HandleReflectorTreeGenerateRow( TSharedPtr<FReflectorNode> InReflectorNode, const TSharedRef<STableViewBase>& OwnerTable );
-
-	/** Callback for getting the child items of the given reflector tree node. */
-	void HandleReflectorTreeGetChildren( TSharedPtr<FReflectorNode> InWidgetGeometry, TArray<TSharedPtr<FReflectorNode>>& OutChildren );
-
-	/** Callback for when the selection in the reflector tree has changed. */
-	void HandleReflectorTreeSelectionChanged( TSharedPtr< FReflectorNode >, ESelectInfo::Type /*SelectInfo*/ )
-	{
-		SelectedNodes = ReflectorTree->GetSelectedItems();
-	}
-
-	TSharedRef<ITableRow> GenerateEventLogRow( TSharedRef<FLoggedEvent> InReflectorNode, const TSharedRef<STableViewBase>& OwnerTable );
-
-private:
-
-	TArray< TSharedRef<FLoggedEvent> > LoggedEvents;
-	TSharedPtr< SListView< TSharedRef< FLoggedEvent > > > EventListView;
-
-	TSharedPtr<SReflectorTree> ReflectorTree;
-
-	TArray<TSharedPtr<FReflectorNode>> SelectedNodes;
-	TArray<TSharedPtr<FReflectorNode>> ReflectorTreeRoot;
-	TArray<TSharedPtr<FReflectorNode>> PickedPath;
-
-	SSplitter::FSlot* WidgetInfoLocation;
-
-	FAccessSourceCode SourceAccessDelegate;
-	FAccessAsset AsseetAccessDelegate;
-
-	bool bShowFocus;
-	bool bIsPicking;
-
-private:
-	// STATS
-	TSharedPtr<SBorder> StatsBorder;
-	TArray< TSharedRef<class FStatItem> > StatsItems;
-	TSharedPtr< SListView< TSharedRef< FStatItem > > > StatsList;
-
-	TSharedRef<SWidget> MakeStatViewer();
-	void UpdateStats();
-	TSharedRef<ITableRow> GenerateStatRow( TSharedRef<FStatItem> StatItem, const TSharedRef<STableViewBase>& OwnerTable );
-	FReply CopyStatsToClipboard();
-
-private:
-	// DEMO MODE
-	bool bEnableDemoMode;
-	double LastMouseClickTime;
-	FVector2D CursorPingPosition;
+	virtual void Construct( const FArguments& InArgs ) = 0;
 };
+
diff --git a/Engine/Source/Developer/StandaloneRenderer/Private/Windows/D3D/SlateD3DRenderer.cpp b/Engine/Source/Developer/StandaloneRenderer/Private/Windows/D3D/SlateD3DRenderer.cpp
--- a/Engine/Source/Developer/StandaloneRenderer/Private/Windows/D3D/SlateD3DRenderer.cpp
+++ b/Engine/Source/Developer/StandaloneRenderer/Private/Windows/D3D/SlateD3DRenderer.cpp
@@ -8,6 +8,11 @@
 #include "Windows/D3D/SlateD3DRenderingPolicy.h"
 #include "ElementBatcher.h"
 #include "FontCache.h"
+#include "SlateStats.h"
+
+SLATE_DECLARE_CYCLE_COUNTER(GRendererDrawElementList, "Renderer DrawElementList");
+SLATE_DECLARE_CYCLE_COUNTER(GRendererUpdateBuffers, "Renderer UpdateBuffers");
+SLATE_DECLARE_CYCLE_COUNTER(GRendererDrawElements, "Renderer DrawElements");
 
 TRefCountPtr<ID3D11Device> GD3DDevice;
 TRefCountPtr<ID3D11DeviceContext> GD3DDeviceContext;
@@ -332,6 +337,7 @@ void FSlateD3DRenderer::DrawWindows( FSlateDrawBuffer& InWindowDrawBuffer )
 	for( int32 ListIndex = 0; ListIndex < WindowElementLists.Num(); ++ListIndex )
 	{
 		FSlateWindowElementList& ElementList = WindowElementLists[ListIndex];
+		SLATE_CYCLE_COUNTER_SCOPE_CUSTOM_DETAILED(SLATE_STATS_DETAIL_LEVEL_MED, GRendererDrawElementList, ElementList.GetWindow()->GetCreatedInFileFName());
 
 		if ( ElementList.GetWindow().IsValid() )
 		{
@@ -351,7 +357,10 @@ void FSlateD3DRenderer::DrawWindows( FSlateDrawBuffer& InWindowDrawBuffer )
 			FSlateD3DViewport* Viewport = WindowToViewportMap.Find( &WindowToDraw.Get() );
 			check(Viewport);
 
-			RenderingPolicy->UpdateBuffers( ElementList );
+			{
+				SLATE_CYCLE_COUNTER_SCOPE(GRendererUpdateBuffers);
+				RenderingPolicy->UpdateBuffers(ElementList);
+			}
 
 			check(Viewport);
 			GD3DDeviceContext->RSSetViewports(1, &Viewport->ViewportInfo );
@@ -367,7 +376,10 @@ void FSlateD3DRenderer::DrawWindows( FSlateDrawBuffer& InWindowDrawBuffer )
 #endif
 			GD3DDeviceContext->OMSetRenderTargets( 1, &RTV, NULL );
 
-			RenderingPolicy->DrawElements( ViewMatrix*Viewport->ProjectionMatrix, ElementList.GetRenderBatches() );
+			{
+				SLATE_CYCLE_COUNTER_SCOPE(GRendererDrawElements);
+				RenderingPolicy->DrawElements(ViewMatrix*Viewport->ProjectionMatrix, ElementList.GetRenderBatches());
+			}
 
 			GD3DDeviceContext->OMSetRenderTargets(0, NULL, NULL);
 
diff --git a/Engine/Source/Developer/StandaloneRenderer/Private/Windows/D3D/SlateD3DRenderingPolicy.cpp b/Engine/Source/Developer/StandaloneRenderer/Private/Windows/D3D/SlateD3DRenderingPolicy.cpp
--- a/Engine/Source/Developer/StandaloneRenderer/Private/Windows/D3D/SlateD3DRenderingPolicy.cpp
+++ b/Engine/Source/Developer/StandaloneRenderer/Private/Windows/D3D/SlateD3DRenderingPolicy.cpp
@@ -7,6 +7,12 @@
 #include "Windows/D3D/SlateD3DRenderingPolicy.h"
 #include "Windows/D3D/SlateD3DTextureManager.h"
 #include "Windows/D3D/SlateD3DTextures.h"
+#include "SlateStats.h"
+
+SLATE_DECLARE_CYCLE_COUNTER(GSlateResizeRenderBuffers, "Resize Render Buffers");
+SLATE_DECLARE_CYCLE_COUNTER(GSlateLockRenderBuffers, "Lock Render Buffers");
+SLATE_DECLARE_CYCLE_COUNTER(GSlateMemCopyRenderBuffers, "Memcopy Render Buffers");
+SLATE_DECLARE_CYCLE_COUNTER(GSlateUnlockRenderBuffers, "Unlock Render Buffers");
 
 /** Offset to apply to UVs to line up texels with pixels */
 const float PixelCenterOffsetD3D11 = 0.0f;
@@ -159,16 +165,27 @@ void FSlateD3D11RenderingPolicy::UpdateBuffers( const FSlateWindowElementList& I
 		// resize if needed
 		if( NumVertices*sizeof(FSlateVertex) > VertexBuffer.GetBufferSize() )
 		{
+			SLATE_CYCLE_COUNTER_SCOPE_DETAILED(SLATE_STATS_DETAIL_LEVEL_FULL, GSlateResizeRenderBuffers);
 			uint32 NumBytesNeeded = NumVertices*sizeof(FSlateVertex);
 			// increase by a static size.
 			// @todo make this better
 			VertexBuffer.ResizeBuffer( NumBytesNeeded + 200*sizeof(FSlateVertex) );
 		}
 
-		void* VerticesPtr = VertexBuffer.Lock(0);
-		FMemory::Memcpy( VerticesPtr, Vertices.GetData(), sizeof(FSlateVertex)*NumVertices );
+		void* VerticesPtr = nullptr;
+		{
+			SLATE_CYCLE_COUNTER_SCOPE_DETAILED(SLATE_STATS_DETAIL_LEVEL_FULL, GSlateLockRenderBuffers);
+			VerticesPtr = VertexBuffer.Lock(0);
+		}
+		{
+			SLATE_CYCLE_COUNTER_SCOPE(GSlateMemCopyRenderBuffers);
+			FMemory::Memcpy(VerticesPtr, Vertices.GetData(), sizeof(FSlateVertex)*NumVertices);
+		}
 
-		VertexBuffer.Unlock();
+		{
+			SLATE_CYCLE_COUNTER_SCOPE_DETAILED(SLATE_STATS_DETAIL_LEVEL_FULL, GSlateUnlockRenderBuffers);
+			VertexBuffer.Unlock();
+		}
 	}
 
 	if( Indices.Num() )
@@ -178,15 +195,26 @@ void FSlateD3D11RenderingPolicy::UpdateBuffers( const FSlateWindowElementList& I
 		// resize if needed
 		if( NumIndices > IndexBuffer.GetMaxNumIndices() )
 		{
+			SLATE_CYCLE_COUNTER_SCOPE_DETAILED(SLATE_STATS_DETAIL_LEVEL_FULL, GSlateResizeRenderBuffers);
 			// increase by a static size.
 			// @todo make this better
 			IndexBuffer.ResizeBuffer( NumIndices + 100 );
 		}
 
-		void* IndicesPtr = IndexBuffer.Lock(0);
-		FMemory::Memcpy( IndicesPtr, Indices.GetData(), sizeof(SlateIndex)*NumIndices );
+		void* IndicesPtr = nullptr;
+		{
+			SLATE_CYCLE_COUNTER_SCOPE_DETAILED(SLATE_STATS_DETAIL_LEVEL_FULL, GSlateLockRenderBuffers);
+			IndicesPtr = IndexBuffer.Lock(0);
+		}
+		{
+			SLATE_CYCLE_COUNTER_SCOPE_DETAILED(SLATE_STATS_DETAIL_LEVEL_FULL, GSlateMemCopyRenderBuffers);
+			FMemory::Memcpy(IndicesPtr, Indices.GetData(), sizeof(SlateIndex)*NumIndices);
+		}
 
-		IndexBuffer.Unlock();
+		{
+			SLATE_CYCLE_COUNTER_SCOPE_DETAILED(SLATE_STATS_DETAIL_LEVEL_FULL, GSlateUnlockRenderBuffers);
+			IndexBuffer.Unlock();
+		}
 	}
 }
 
diff --git a/Engine/Source/Editor/GraphEditor/Private/SGraphPanel.cpp b/Engine/Source/Editor/GraphEditor/Private/SGraphPanel.cpp
--- a/Engine/Source/Editor/GraphEditor/Private/SGraphPanel.cpp
+++ b/Engine/Source/Editor/GraphEditor/Private/SGraphPanel.cpp
@@ -30,8 +30,6 @@
 
 DEFINE_LOG_CATEGORY_STATIC(LogGraphPanel, Log, All);
 
-DECLARE_CYCLE_STAT( TEXT("OnPaint SGraphPanel"), STAT_SlateOnPaint_SGraphPanel, STATGROUP_Slate );
-
 SGraphPanel::FGraphPinHandle::FGraphPinHandle(UEdGraphPin* InPin)
 {
 	if (auto* Node = InPin->GetOwningNode())
@@ -116,10 +114,6 @@ SGraphPanel::~SGraphPanel()
 
 int32 SGraphPanel::OnPaint( const FPaintArgs& Args, const FGeometry& AllottedGeometry, const FSlateRect& MyClippingRect, FSlateWindowElementList& OutDrawElements, int32 LayerId, const FWidgetStyle& InWidgetStyle, bool bParentEnabled ) const
 {
-#if SLATE_HD_STATS
-	SCOPE_CYCLE_COUNTER( STAT_SlateOnPaint_SGraphPanel );
-#endif
-
 	CachedAllottedGeometryScaledSize = AllottedGeometry.Size * AllottedGeometry.Scale;
 
 	//Style used for objects that are the same between revisions
diff --git a/Engine/Source/Programs/SlateViewer/SlateViewer.Target.cs b/Engine/Source/Programs/SlateViewer/SlateViewer.Target.cs
--- a/Engine/Source/Programs/SlateViewer/SlateViewer.Target.cs
+++ b/Engine/Source/Programs/SlateViewer/SlateViewer.Target.cs
@@ -61,8 +61,6 @@ public class SlateViewerTarget : TargetRules
 		// We need CoreUObject compiled in as the source code access module requires it
 		UEBuildConfiguration.bCompileAgainstCoreUObject = true;
 
-		UEBuildConfiguration.bCompileWithStatsWithoutEngine = true;
-
 		// SlateViewer.exe has no exports, so no need to verify that a .lib and .exp file was emitted by
 		// the linker.
 		OutLinkEnvironmentConfiguration.bHasExports = false;
diff --git a/Engine/Source/Runtime/Slate/Private/Framework/Application/SlateApplication.cpp b/Engine/Source/Runtime/Slate/Private/Framework/Application/SlateApplication.cpp
--- a/Engine/Source/Runtime/Slate/Private/Framework/Application/SlateApplication.cpp
+++ b/Engine/Source/Runtime/Slate/Private/Framework/Application/SlateApplication.cpp
@@ -4,6 +4,7 @@
 #include "SlatePrivatePCH.h"
 #include "SWindowTitleBar.h"
 #include "HittestGrid.h"
+#include "SlateStats.h"
 
 #include "IWidgetReflector.h"
 #include "GenericCommands.h"
@@ -268,11 +269,23 @@ DECLARE_CYCLE_STAT( TEXT("Message Tick Time"), STAT_SlateMessageTick, STATGROUP_
 DECLARE_CYCLE_STAT( TEXT("Update Tooltip Time"), STAT_SlateUpdateTooltip, STATGROUP_Slate );
 DECLARE_CYCLE_STAT( TEXT("Tick Window And Children Time"), STAT_SlateTickWindowAndChildren, STATGROUP_Slate );
 DECLARE_CYCLE_STAT( TEXT("Total Slate Tick Time"), STAT_SlateTickTime, STATGROUP_Slate );
-DECLARE_CYCLE_STAT( TEXT("SlatePrepass"), STAT_SlateCacheDesiredSize, STATGROUP_Slate );
+DECLARE_CYCLE_STAT( TEXT("SlatePrepass"), STAT_SlatePrepass, STATGROUP_Slate );
 DECLARE_CYCLE_STAT( TEXT("DrawWindows"), STAT_SlateDrawWindowTime, STATGROUP_Slate );
 DECLARE_CYCLE_STAT( TEXT("TickWidgets"), STAT_SlateTickWidgets, STATGROUP_Slate );
 DECLARE_CYCLE_STAT( TEXT("TickRegisteredWidgets"), STAT_SlateTickRegisteredWidgets, STATGROUP_Slate );
 
+SLATE_DECLARE_CYCLE_COUNTER(GSlateTotalTickTime, "Total Slate Tick Time");
+SLATE_DECLARE_CYCLE_COUNTER(GMessageTickTime, "Message Tick Time");
+SLATE_DECLARE_CYCLE_COUNTER(GUpdateTooltipTime, "Update Tooltip Time");
+SLATE_DECLARE_CYCLE_COUNTER(GSlateSynthesizeMouseMove, "Synthesize Mouse Move");
+SLATE_DECLARE_CYCLE_COUNTER(GTickWindowAndChildrenTime, "Tick Window And Children Time");
+SLATE_DECLARE_CYCLE_COUNTER(GTickWidgets, "TickWidgets");
+SLATE_DECLARE_CYCLE_COUNTER(GSlateTickNotificationManager, "NotificationManager Tick");
+SLATE_DECLARE_CYCLE_COUNTER(GSlateDrawWindows, "DrawWindows");
+SLATE_DECLARE_CYCLE_COUNTER(GSlateDrawWindowAndChildren, "Draw Window And Children");
+SLATE_DECLARE_CYCLE_COUNTER(GSlateRendererDrawWindows, "Renderer DrawWindows");
+SLATE_DECLARE_CYCLE_COUNTER(GSlateDrawPrepass, "DrawPrepass");
+SLATE_DECLARE_CYCLE_COUNTER(GSlatePrepassWindowAndChildren, "Prepass Window And Children");
 
 // Slate Event Logging is enabled to allow crash log dumping
 #define LOG_SLATE_EVENTS 0
@@ -873,6 +886,7 @@ bool FSlateApplication::IsWindowHousingInteractiveTooltip(const TSharedRef<const
 
 void FSlateApplication::TickWindowAndChildren( TSharedRef<SWindow> WindowToTick )
 {
+	SLATE_CYCLE_COUNTER_SCOPE_CUSTOM(GTickWindowAndChildrenTime, WindowToTick->GetCreatedInFileFName());
 	if ( WindowToTick->IsVisible() && !WindowToTick->IsWindowMinimized() )
 	{
 		// Switch to the appropriate world for ticking
@@ -883,7 +897,7 @@ void FSlateApplication::TickWindowAndChildren( TSharedRef<SWindow> WindowToTick
 		// the Tick() function.
 		
 		{
-			SCOPE_CYCLE_COUNTER( STAT_SlateCacheDesiredSize );
+			SCOPE_CYCLE_COUNTER( STAT_SlatePrepass );
 			WindowToTick->SlatePrepass(GetApplicationScale());
 		}
 
@@ -894,6 +908,7 @@ void FSlateApplication::TickWindowAndChildren( TSharedRef<SWindow> WindowToTick
 
 		{
 			SCOPE_CYCLE_COUNTER( STAT_SlateTickWidgets );
+			SLATE_CYCLE_COUNTER_SCOPE(GTickWidgets);
 			// Tick this window and all of the widgets in this window
 			WindowToTick->TickWidgetsRecursively( WindowToTick->GetWindowGeometryInScreen(), GetCurrentTime(), GetDeltaTime() );
 		}
@@ -909,6 +924,7 @@ void FSlateApplication::TickWindowAndChildren( TSharedRef<SWindow> WindowToTick
 
 void FSlateApplication::DrawWindows()
 {
+	SLATE_CYCLE_COUNTER_SCOPE(GSlateDrawWindows);
 	PrivateDrawWindows();
 }
 
@@ -927,6 +943,7 @@ struct FDrawWindowArgs
 
 void FSlateApplication::DrawWindowAndChildren( const TSharedRef<SWindow>& WindowToDraw, FDrawWindowArgs& DrawWindowArgs )
 {
+	SLATE_CYCLE_COUNTER_SCOPE_CUSTOM(GSlateDrawWindowAndChildren, WindowToDraw->GetCreatedInFileFName());
 	// Only draw visible windows
 	if( WindowToDraw->IsVisible() && !WindowToDraw->IsWindowMinimized() )
 	{
@@ -1006,10 +1023,11 @@ void FSlateApplication::DrawWindowAndChildren( const TSharedRef<SWindow>& Window
 
 static void PrepassWindowAndChildren( TSharedRef<SWindow> WindowToPrepass )
 {
+	SLATE_CYCLE_COUNTER_SCOPE_CUSTOM(GSlatePrepassWindowAndChildren, WindowToPrepass->GetCreatedInFileFName());
 	FScopedSwitchWorldHack SwitchWorld( WindowToPrepass );
 
 	{
-		SCOPE_CYCLE_COUNTER(STAT_SlateCacheDesiredSize);
+		SCOPE_CYCLE_COUNTER(STAT_SlatePrepass);
 		WindowToPrepass->SlatePrepass(FSlateApplication::Get().GetApplicationScale());
 	}
 
@@ -1026,6 +1044,7 @@ static void PrepassWindowAndChildren( TSharedRef<SWindow> WindowToPrepass )
 
 void FSlateApplication::DrawPrepass( TSharedPtr<SWindow> DrawOnlyThisWindow )
 {
+	SLATE_CYCLE_COUNTER_SCOPE(GSlateDrawPrepass);
 	TSharedPtr<SWindow> ActiveModalWindow = GetActiveModalWindow();
 
 	if (ActiveModalWindow.IsValid())
@@ -1148,8 +1167,11 @@ void FSlateApplication::PrivateDrawWindows( TSharedPtr<SWindow> DrawOnlyThisWind
 		TSharedPtr<SWindow> CandidateWindow = Candidate.GetWindow();
 		return !CandidateWindow.IsValid() || !AllWindows.Contains(CandidateWindow.ToSharedRef());
 	});
-	
-	Renderer->DrawWindows( DrawWindowArgs.OutDrawBuffer );
+
+	{	
+		SLATE_CYCLE_COUNTER_SCOPE(GSlateRendererDrawWindows);
+		Renderer->DrawWindows( DrawWindowArgs.OutDrawBuffer );
+	}
 }
 
 void FSlateApplication::PollGameDeviceState()
@@ -1224,7 +1246,9 @@ extern SLATECORE_API int32 bFoldTick;
 
 void FSlateApplication::Tick()
 {
+	{
 	SCOPE_CYCLE_COUNTER( STAT_SlateTickTime );
+	SLATE_CYCLE_COUNTER_SCOPE(GSlateTotalTickTime);
 	
 	{
 		const float DeltaTime = GetDeltaTime();
@@ -1262,6 +1286,7 @@ void FSlateApplication::Tick()
 
 	{
 		SCOPE_CYCLE_COUNTER( STAT_SlateUpdateTooltip );
+		SLATE_CYCLE_COUNTER_SCOPE(GUpdateTooltipTime);
 
 		// Update tool tip, if we have one
 		const bool AllowSpawningOfToolTips = false;
@@ -1336,51 +1361,57 @@ void FSlateApplication::Tick()
 	{
 		if (!bFoldTick)
 		{
+			SCOPE_CYCLE_COUNTER(STAT_SlateTickWindowAndChildren);
+			SLATE_CYCLE_COUNTER_SCOPE(GTickWindowAndChildrenTime);
+
 			TSharedPtr<SWindow> ActiveModalWindow = GetActiveModalWindow();
 
+			if ( ActiveModalWindow.IsValid() )
 			{
-				SCOPE_CYCLE_COUNTER(STAT_SlateTickWindowAndChildren);
-
-				if ( ActiveModalWindow.IsValid() )
+				// There is a modal window, and we just need to tick it.
+				TickWindowAndChildren(ActiveModalWindow.ToSharedRef());
+				// And also tick any topmost windows (like tooltips, etc)
+				for ( TArray< TSharedRef<SWindow> >::TIterator CurrentWindowIt(SlateWindows); CurrentWindowIt; ++CurrentWindowIt )
 				{
-					// There is a modal window, and we just need to tick it.
-					TickWindowAndChildren(ActiveModalWindow.ToSharedRef());
-					// And also tick any topmost windows (like tooltips, etc)
-					for ( TArray< TSharedRef<SWindow> >::TIterator CurrentWindowIt(SlateWindows); CurrentWindowIt; ++CurrentWindowIt )
-					{
-						TSharedRef<SWindow>& CurrentWindow = *CurrentWindowIt;
-						if ( CurrentWindow->IsTopmostWindow() )
-						{
-							TickWindowAndChildren(CurrentWindow);
-						}
-					}
-					// also tick the notification manager's windows
-					TArray< TSharedRef<SWindow> > NotificationWindows;
-					FSlateNotificationManager::Get().GetWindows(NotificationWindows);
-					for ( auto CurrentWindowIt(NotificationWindows.CreateIterator()); CurrentWindowIt; ++CurrentWindowIt )
+					TSharedRef<SWindow>& CurrentWindow = *CurrentWindowIt;
+					if ( CurrentWindow->IsTopmostWindow() )
 					{
-						TickWindowAndChildren(*CurrentWindowIt);
+						TickWindowAndChildren(CurrentWindow);
 					}
 				}
-				else
+				// also tick the notification manager's windows
+				TArray< TSharedRef<SWindow> > NotificationWindows;
+				FSlateNotificationManager::Get().GetWindows(NotificationWindows);
+				for ( auto CurrentWindowIt(NotificationWindows.CreateIterator()); CurrentWindowIt; ++CurrentWindowIt )
 				{
-					// No modal window; tick all slate windows.
-					for ( TArray< TSharedRef<SWindow> >::TIterator CurrentWindowIt(SlateWindows); CurrentWindowIt; ++CurrentWindowIt )
-					{
-						TSharedRef<SWindow>& CurrentWindow = *CurrentWindowIt;
-						TickWindowAndChildren(CurrentWindow);
-					}
+					TickWindowAndChildren(*CurrentWindowIt);
+				}
+			}
+			else
+			{
+				// No modal window; tick all slate windows.
+				for ( TArray< TSharedRef<SWindow> >::TIterator CurrentWindowIt(SlateWindows); CurrentWindowIt; ++CurrentWindowIt )
+				{
+					TSharedRef<SWindow>& CurrentWindow = *CurrentWindowIt;
+					TickWindowAndChildren(CurrentWindow);
 				}
 			}
 		}
 
 		// Update any notifications - this needs to be done after windows have updated themselves 
 		// (so they know their size)
-		FSlateNotificationManager::Get().Tick();
+		{
+			SLATE_CYCLE_COUNTER_SCOPE(GSlateTickNotificationManager);
+			FSlateNotificationManager::Get().Tick();
+		}
 
 		// Draw all windows
 		DrawWindows();
 	}
+	}
+
+	// Update Slate Stats
+	SLATE_STATS_END_FRAME(GetCurrentTime());
 }
 
 
@@ -3066,6 +3097,7 @@ bool FSlateApplication::IsWindowInDestroyQueue(TSharedRef<SWindow> Window) const
 
 void FSlateApplication::SynthesizeMouseMove()
 {
+	SLATE_CYCLE_COUNTER_SCOPE(GSlateSynthesizeMouseMove);
 	if (PlatformApplication->Cursor.IsValid())
 	{
 		// Synthetic mouse events accomplish two goals:
diff --git a/Engine/Source/Runtime/Slate/Private/Widgets/Images/SImage.cpp b/Engine/Source/Runtime/Slate/Private/Widgets/Images/SImage.cpp
--- a/Engine/Source/Runtime/Slate/Private/Widgets/Images/SImage.cpp
+++ b/Engine/Source/Runtime/Slate/Private/Widgets/Images/SImage.cpp
@@ -2,8 +2,6 @@
 
 #include "SlatePrivatePCH.h"
 
-DECLARE_CYCLE_STAT( TEXT("OnPaint SImage"), STAT_SlateOnPaint_SImage, STATGROUP_Slate );
-
 /**
  * Construct this widget
  *
@@ -19,9 +17,6 @@ void SImage::Construct( const FArguments& InArgs )
 
 int32 SImage::OnPaint( const FPaintArgs& Args, const FGeometry& AllottedGeometry, const FSlateRect& MyClippingRect, FSlateWindowElementList& OutDrawElements, int32 LayerId, const FWidgetStyle& InWidgetStyle, bool bParentEnabled ) const
 {
-#if SLATE_HD_STATS
-	SCOPE_CYCLE_COUNTER( STAT_SlateOnPaint_SImage );
-#endif
 	const FSlateBrush* ImageBrush = Image.Get();
 
 	if ((ImageBrush != nullptr) && (ImageBrush->DrawAs != ESlateBrushDrawType::NoDrawType))
diff --git a/Engine/Source/Runtime/Slate/Private/Widgets/Input/SEditableText.cpp b/Engine/Source/Runtime/Slate/Private/Widgets/Input/SEditableText.cpp
--- a/Engine/Source/Runtime/Slate/Private/Widgets/Input/SEditableText.cpp
+++ b/Engine/Source/Runtime/Slate/Private/Widgets/Input/SEditableText.cpp
@@ -4,8 +4,6 @@
 #include "TextEditHelper.h"
 #include "GenericCommands.h"
 
-DECLARE_CYCLE_STAT( TEXT("OnPaint SEditableText"), STAT_SlateOnPaint_SEditableText, STATGROUP_Slate );
-
 /** A pointer to the editable text currently under the mouse cursor. nullptr when there isn't one. */
 SEditableText* SEditableText::EditableTextUnderCursor = nullptr;
 
@@ -1245,9 +1243,6 @@ FVector2D SEditableText::ComputeDesiredSize(float LayoutScaleMultiplier) const
 
 int32 SEditableText::OnPaint( const FPaintArgs& Args, const FGeometry& AllottedGeometry, const FSlateRect& MyClippingRect, FSlateWindowElementList& OutDrawElements, int32 LayerId, const FWidgetStyle& InWidgetStyle, bool bParentEnabled ) const
 {
-#if SLATE_HD_STATS
-	SCOPE_CYCLE_COUNTER( STAT_SlateOnPaint_SEditableText );
-#endif
 	// The text and some effects draws in front of the widget's background and selection.
 	const int32 SelectionLayer = 0;
 	const int32 TextLayer = 1;
diff --git a/Engine/Source/Runtime/Slate/Private/Widgets/Input/SSlider.cpp b/Engine/Source/Runtime/Slate/Private/Widgets/Input/SSlider.cpp
--- a/Engine/Source/Runtime/Slate/Private/Widgets/Input/SSlider.cpp
+++ b/Engine/Source/Runtime/Slate/Private/Widgets/Input/SSlider.cpp
@@ -2,8 +2,6 @@
 
 #include "SlatePrivatePCH.h"
 
-DECLARE_CYCLE_STAT( TEXT("OnPaint SSlider"), STAT_SlateOnPaint_SSlider, STATGROUP_Slate );
-
 void SSlider::Construct( const SSlider::FArguments& InDeclaration )
 {
 	check(InDeclaration._Style);
@@ -24,10 +22,6 @@ void SSlider::Construct( const SSlider::FArguments& InDeclaration )
 
 int32 SSlider::OnPaint( const FPaintArgs& Args, const FGeometry& AllottedGeometry, const FSlateRect& MyClippingRect, FSlateWindowElementList& OutDrawElements, int32 LayerId, const FWidgetStyle& InWidgetStyle, bool bParentEnabled ) const
 {
-#if SLATE_HD_STATS
-	SCOPE_CYCLE_COUNTER(STAT_SlateOnPaint_SSlider);
-#endif
-
 	// we draw the slider like a horizontal slider regardless of the orientation, and apply a render transform to make it display correctly.
 	// However, the AllottedGeometry is computed as it will be rendered, so we have to use the "horizontal orientation" when doing drawing computations.
 	const float AllottedWidth = Orientation == Orient_Horizontal ? AllottedGeometry.GetLocalSize().X : AllottedGeometry.GetLocalSize().Y;
diff --git a/Engine/Source/Runtime/Slate/Private/Widgets/Layout/SBorder.cpp b/Engine/Source/Runtime/Slate/Private/Widgets/Layout/SBorder.cpp
--- a/Engine/Source/Runtime/Slate/Private/Widgets/Layout/SBorder.cpp
+++ b/Engine/Source/Runtime/Slate/Private/Widgets/Layout/SBorder.cpp
@@ -2,8 +2,6 @@
 
 #include "SlatePrivatePCH.h"
  
-DECLARE_CYCLE_STAT( TEXT("OnPaint SBorder"), STAT_SlateOnPaint_SBorder, STATGROUP_Slate );
-
 SBorder::SBorder()
 	: BorderImage( FCoreStyle::Get().GetBrush( "Border" ) )
 	, BorderBackgroundColor( FLinearColor::White )
@@ -69,28 +67,23 @@ void SBorder::ClearContent()
 int32 SBorder::OnPaint( const FPaintArgs& Args, const FGeometry& AllottedGeometry, const FSlateRect& MyClippingRect, FSlateWindowElementList& OutDrawElements, int32 LayerId, const FWidgetStyle& InWidgetStyle, bool bParentEnabled ) const
 {
 	bool bEnabled = false;
-	{
-#if SLATE_HD_STATS
-		SCOPE_CYCLE_COUNTER( STAT_SlateOnPaint_SBorder );
-#endif
-		const FSlateBrush* BrushResource = BorderImage.Get();
+	const FSlateBrush* BrushResource = BorderImage.Get();
 		
-		bEnabled = ShouldBeEnabled( bParentEnabled );
-		bool bShowDisabledEffect = ShowDisabledEffect.Get();
-		ESlateDrawEffect::Type DrawEffects = bShowDisabledEffect && !bEnabled ? ESlateDrawEffect::DisabledEffect : ESlateDrawEffect::None;
-
-		if ( BrushResource && BrushResource->DrawAs != ESlateBrushDrawType::NoDrawType )
-		{
-			FSlateDrawElement::MakeBox(
-				OutDrawElements,
-				LayerId,
-				AllottedGeometry.ToPaintGeometry(),
-				BrushResource,
-				MyClippingRect,
-				DrawEffects,
-				BrushResource->GetTint( InWidgetStyle ) * InWidgetStyle.GetColorAndOpacityTint() * BorderBackgroundColor.Get().GetColor( InWidgetStyle )
-			);
-		}
+	bEnabled = ShouldBeEnabled( bParentEnabled );
+	bool bShowDisabledEffect = ShowDisabledEffect.Get();
+	ESlateDrawEffect::Type DrawEffects = bShowDisabledEffect && !bEnabled ? ESlateDrawEffect::DisabledEffect : ESlateDrawEffect::None;
+
+	if ( BrushResource && BrushResource->DrawAs != ESlateBrushDrawType::NoDrawType )
+	{
+		FSlateDrawElement::MakeBox(
+			OutDrawElements,
+			LayerId,
+			AllottedGeometry.ToPaintGeometry(),
+			BrushResource,
+			MyClippingRect,
+			DrawEffects,
+			BrushResource->GetTint( InWidgetStyle ) * InWidgetStyle.GetColorAndOpacityTint() * BorderBackgroundColor.Get().GetColor( InWidgetStyle )
+		);
 	}
 
 	FWidgetStyle CompoundedWidgetStyle = FWidgetStyle(InWidgetStyle)
diff --git a/Engine/Source/Runtime/Slate/Private/Widgets/Layout/SBox.cpp b/Engine/Source/Runtime/Slate/Private/Widgets/Layout/SBox.cpp
--- a/Engine/Source/Runtime/Slate/Private/Widgets/Layout/SBox.cpp
+++ b/Engine/Source/Runtime/Slate/Private/Widgets/Layout/SBox.cpp
@@ -3,8 +3,6 @@
 #include "SlatePrivatePCH.h"
 #include "LayoutUtils.h"
 
-DECLARE_CYCLE_STAT( TEXT("OnPaint SBox"), STAT_SlateOnPaint_SBox, STATGROUP_Slate );
-
 SBox::SBox()
 : ChildSlot()
 {
@@ -158,12 +156,7 @@ int32 SBox::OnPaint( const FPaintArgs& Args, const FGeometry& AllottedGeometry,
 {
 	// An SBox just draws its only child
 	FArrangedChildren ArrangedChildren(EVisibility::Visible);
-	{
-#if SLATE_HD_STATS
-		SCOPE_CYCLE_COUNTER( STAT_SlateOnPaint_SBox );
-#endif
-		this->ArrangeChildren(AllottedGeometry, ArrangedChildren);
-	}
+	this->ArrangeChildren(AllottedGeometry, ArrangedChildren);
 
 	// Maybe none of our children are visible
 	if( ArrangedChildren.Num() > 0 )
diff --git a/Engine/Source/Runtime/Slate/Private/Widgets/SViewport.cpp b/Engine/Source/Runtime/Slate/Private/Widgets/SViewport.cpp
--- a/Engine/Source/Runtime/Slate/Private/Widgets/SViewport.cpp
+++ b/Engine/Source/Runtime/Slate/Private/Widgets/SViewport.cpp
@@ -3,8 +3,6 @@
 #include "SlatePrivatePCH.h"
 #include "HittestGrid.h"
 
-DECLARE_CYCLE_STAT( TEXT("OnPaint SViewport"), STAT_SlateOnPaint_SViewport, STATGROUP_Slate );
-
 /* SViewport constructors
  *****************************************************************************/
 
@@ -54,9 +52,6 @@ EActiveTimerReturnType SViewport::EnsureTick(double InCurrentTime, float InDelta
 
 int32 SViewport::OnPaint( const FPaintArgs& Args, const FGeometry& AllottedGeometry, const FSlateRect& MyClippingRect, FSlateWindowElementList& OutDrawElements, int32 LayerId, const FWidgetStyle& InWidgetStyle, bool bParentEnabled ) const
 {
-#if SLATE_HD_STATS
-	SCOPE_CYCLE_COUNTER( STAT_SlateOnPaint_SViewport );
-#endif
 	bool bEnabled = ShouldBeEnabled( bParentEnabled );
 	bool bShowDisabledEffect = ShowDisabledEffect.Get();
 	ESlateDrawEffect::Type DrawEffects = bShowDisabledEffect && !bEnabled ? ESlateDrawEffect::DisabledEffect : ESlateDrawEffect::None;
diff --git a/Engine/Source/Runtime/Slate/Private/Widgets/Text/SRichTextBlock.cpp b/Engine/Source/Runtime/Slate/Private/Widgets/Text/SRichTextBlock.cpp
--- a/Engine/Source/Runtime/Slate/Private/Widgets/Text/SRichTextBlock.cpp
+++ b/Engine/Source/Runtime/Slate/Private/Widgets/Text/SRichTextBlock.cpp
@@ -4,8 +4,6 @@
 
 #if WITH_FANCY_TEXT
 
-DECLARE_CYCLE_STAT( TEXT("OnPaint SRichTextBlock"), STAT_SlateOnPaint_SRichTextBlock, STATGROUP_Slate );
-
 #include "TextLayoutEngine.h"
 #include "TextBlockLayout.h"
 #include "IRichTextMarkupParser.h"
@@ -43,8 +41,6 @@ void SRichTextBlock::Construct( const FArguments& InArgs )
 
 int32 SRichTextBlock::OnPaint( const FPaintArgs& Args, const FGeometry& AllottedGeometry, const FSlateRect& MyClippingRect, FSlateWindowElementList& OutDrawElements, int32 LayerId, const FWidgetStyle& InWidgetStyle, bool bParentEnabled ) const
 {
-	SCOPE_CYCLE_COUNTER( STAT_SlateOnPaint_SRichTextBlock );
-
 	// OnPaint will also update the text layout cache if required
 	LayerId = TextLayoutCache->OnPaint(Args, AllottedGeometry, MyClippingRect, OutDrawElements, LayerId, InWidgetStyle, ShouldBeEnabled(bParentEnabled));
 
diff --git a/Engine/Source/Runtime/Slate/Private/Widgets/Text/STextBlock.cpp b/Engine/Source/Runtime/Slate/Private/Widgets/Text/STextBlock.cpp
--- a/Engine/Source/Runtime/Slate/Private/Widgets/Text/STextBlock.cpp
+++ b/Engine/Source/Runtime/Slate/Private/Widgets/Text/STextBlock.cpp
@@ -4,9 +4,6 @@
 #include "PlainTextLayoutMarshaller.h"
 #include "TextBlockLayout.h"
 
-DECLARE_CYCLE_STAT( TEXT("STextBlock"), STAT_SlateComputeDesiredSize_STextBlock, STATGROUP_Slate );
-DECLARE_CYCLE_STAT( TEXT("OnPaint STextBlock"), STAT_SlateOnPaint_STextBlock, STATGROUP_Slate );
-
 void STextBlock::Construct( const FArguments& InArgs )
 {
 	TextStyle = InArgs._TextStyle;
@@ -98,10 +95,6 @@ void STextBlock::SetText( const FText& InText )
 
 int32 STextBlock::OnPaint( const FPaintArgs& Args, const FGeometry& AllottedGeometry, const FSlateRect& MyClippingRect, FSlateWindowElementList& OutDrawElements, int32 LayerId, const FWidgetStyle& InWidgetStyle, bool bParentEnabled ) const
 {
-#if SLATE_HD_STATS
-	SCOPE_CYCLE_COUNTER( STAT_SlateOnPaint_STextBlock );
-#endif
-
 #if WITH_FANCY_TEXT
 
 	// OnPaint will also update the text layout cache if required
@@ -170,10 +163,6 @@ FReply STextBlock::OnMouseButtonDoubleClick( const FGeometry& InMyGeometry, cons
 
 FVector2D STextBlock::ComputeDesiredSize(float LayoutScaleMultiplier) const
 {
-#if SLATE_HD_STATS
-	SCOPE_CYCLE_COUNTER( STAT_SlateComputeDesiredSize_STextBlock );
-#endif
-
 #if WITH_FANCY_TEXT
 
 	// todo: jdale - The scale needs to be passed to ComputeDesiredSize
diff --git a/Engine/Source/Runtime/Slate/Private/Widgets/Text/TextBlockLayout.cpp b/Engine/Source/Runtime/Slate/Private/Widgets/Text/TextBlockLayout.cpp
--- a/Engine/Source/Runtime/Slate/Private/Widgets/Text/TextBlockLayout.cpp
+++ b/Engine/Source/Runtime/Slate/Private/Widgets/Text/TextBlockLayout.cpp
@@ -5,9 +5,12 @@
 #include "SlateTextLayout.h"
 #include "BaseTextLayoutMarshaller.h"
 #include "SlateTextHighlightRunRenderer.h"
+#include "SlateStats.h"
 
 #if WITH_FANCY_TEXT
 
+SLATE_DECLARE_CYCLE_COUNTER(GSlateTextBlockLayoutComputeDesiredSize, "FTextBlockLayout ComputeDesiredSize");
+
 TSharedRef<FTextBlockLayout> FTextBlockLayout::Create(FTextBlockStyle InDefaultTextStyle, TSharedRef<FBaseTextLayoutMarshaller> InMarshaller, TSharedPtr<IBreakIterator> InLineBreakPolicy)
 {
 	return MakeShareable(new FTextBlockLayout(MoveTemp(InDefaultTextStyle), InMarshaller, InLineBreakPolicy));
@@ -24,6 +27,7 @@ FTextBlockLayout::FTextBlockLayout(FTextBlockStyle InDefaultTextStyle, TSharedRe
 
 FVector2D FTextBlockLayout::ComputeDesiredSize(const FWidgetArgs& InWidgetArgs, const float InScale, const FTextBlockStyle& InTextStyle)
 {
+	SLATE_CYCLE_COUNTER_SCOPE_DETAILED(SLATE_STATS_DETAIL_LEVEL_HI, GSlateTextBlockLayoutComputeDesiredSize);
 	TextLayout->SetWrappingWidth(CalculateWrappingWidth(InWidgetArgs));
 	TextLayout->SetMargin(InWidgetArgs.Margin.Get());
 	TextLayout->SetJustification(InWidgetArgs.Justification.Get());
diff --git a/Engine/Source/Runtime/Slate/Public/SlateBasics.h b/Engine/Source/Runtime/Slate/Public/SlateBasics.h
--- a/Engine/Source/Runtime/Slate/Public/SlateBasics.h
+++ b/Engine/Source/Runtime/Slate/Public/SlateBasics.h
@@ -9,10 +9,6 @@
 #include "SlateClasses.h"
 
 
-// Enables or disables more slate stats.  These stats are usually per widget and can be gathered hundreds or thousands of times per frame
-// enabling these constantly causes a bit of a performance drop due to stat gathering overhead so these should only be enabled for profiling
-#define SLATE_HD_STATS 0
-
 #define BEGIN_SLATE_FUNCTION_BUILD_OPTIMIZATION PRAGMA_DISABLE_OPTIMIZATION
 #define END_SLATE_FUNCTION_BUILD_OPTIMIZATION   PRAGMA_ENABLE_OPTIMIZATION
 
diff --git a/Engine/Source/Runtime/SlateCore/Private/Fonts/FontCache.cpp b/Engine/Source/Runtime/SlateCore/Private/Fonts/FontCache.cpp
--- a/Engine/Source/Runtime/SlateCore/Private/Fonts/FontCache.cpp
+++ b/Engine/Source/Runtime/SlateCore/Private/Fonts/FontCache.cpp
@@ -3,6 +3,9 @@
 #include "SlateCorePrivatePCH.h"
 #include "LegacySlateFontInfoCache.h"
 
+DECLARE_DWORD_ACCUMULATOR_STAT(TEXT("Num Font Atlases"), STAT_SlateNumFontAtlases, STATGROUP_SlateMemory);
+DECLARE_MEMORY_STAT(TEXT("Font Kerning Table Memory"), STAT_SlateFontKerningTableMemory, STATGROUP_SlateMemory);
+DEFINE_STAT(STAT_SlateFontMeasureCacheMemory);
 
 #ifndef WITH_FREETYPE
 	#define WITH_FREETYPE	0
@@ -904,18 +907,16 @@ int8 FKerningTable::GetKerning( const FFontData& InFontData, const int32 InSize,
 		{
 			OutKerning = FontCache.GetKerning( InFontData, InSize, FirstChar, SecondChar, InScale );
 
-#if STATS
-			const uint32 CurrentMemoryUsage = MappedKerningPairs.GetAllocatedSize();
-#endif
+			STAT(const uint32 CurrentMemoryUsage = MappedKerningPairs.GetAllocatedSize());
 			MappedKerningPairs.Add( KerningPair, OutKerning );
-		
-#if STATS
-			uint32 NewMemoryUsage = MappedKerningPairs.GetAllocatedSize();
-			if( NewMemoryUsage > CurrentMemoryUsage )
+			STAT(
 			{
-				INC_MEMORY_STAT_BY( STAT_SlateFontKerningTableMemory, NewMemoryUsage-CurrentMemoryUsage );
-			}
-#endif
+				uint32 NewMemoryUsage = MappedKerningPairs.GetAllocatedSize();
+				if (NewMemoryUsage > CurrentMemoryUsage)
+				{
+					INC_MEMORY_STAT_BY(STAT_SlateFontKerningTableMemory, NewMemoryUsage - CurrentMemoryUsage);
+				}
+			})
 		}
 	}
 
@@ -1040,8 +1041,6 @@ const FCharacterEntry& FCharacterList::GetCharacter( TCHAR Character )
 
 FCharacterEntry& FCharacterList::CacheCharacter( TCHAR Character )
 {
-	SCOPE_CYCLE_COUNTER( STAT_SlateFontCachingTime );
-
 	FCharacterEntry NewEntry;
 	bool bSuccess = FontCache.AddNewEntry( Character, FontKey, NewEntry );
 
diff --git a/Engine/Source/Runtime/SlateCore/Private/Fonts/FontCacheUtils.h b/Engine/Source/Runtime/SlateCore/Private/Fonts/FontCacheUtils.h
--- a/Engine/Source/Runtime/SlateCore/Private/Fonts/FontCacheUtils.h
+++ b/Engine/Source/Runtime/SlateCore/Private/Fonts/FontCacheUtils.h
@@ -6,6 +6,7 @@
 typedef FString KeyType;
 typedef FVector2D ValueType;
 
+DECLARE_MEMORY_STAT_EXTERN(TEXT("Font Measure Memory"), STAT_SlateFontMeasureCacheMemory, STATGROUP_SlateMemory, SLATECORE_API);
 
 /**
  * Basic Least Recently Used (LRU) cache
diff --git a/Engine/Source/Runtime/SlateCore/Private/Fonts/FontMeasure.cpp b/Engine/Source/Runtime/SlateCore/Private/Fonts/FontMeasure.cpp
--- a/Engine/Source/Runtime/SlateCore/Private/Fonts/FontMeasure.cpp
+++ b/Engine/Source/Runtime/SlateCore/Private/Fonts/FontMeasure.cpp
@@ -1,9 +1,12 @@
 // Copyright 1998-2015 Epic Games, Inc. All Rights Reserved.
 
 #include "SlateCorePrivatePCH.h"
+#include "SlateStats.h"
 
 #define USE_MEASURE_CACHING 1
 
+SLATE_DECLARE_CYCLE_COUNTER(GSlateMeasureStringTime, "Measure String");
+
 namespace FontMeasureConstants
 {
 	/** Number of possible elements in each measurement cache */
@@ -129,7 +132,7 @@ int32 FSlateFontMeasure::FindCharacterIndexAtOffset( const FString& Text, int32
 
 FVector2D FSlateFontMeasure::MeasureStringInternal( const FString& Text, int32 StartIndex, int32 EndIndex, const FSlateFontInfo& InFontInfo, bool IncludeKerningWithPrecedingChar, float FontScale, int32 StopAfterHorizontalOffset, ELastCharacterIndexFormat CharIndexFormat, int32& OutLastCharacterIndex ) const
 {
-	SCOPE_CYCLE_COUNTER(STAT_SlateMeasureStringTime);
+	SLATE_CYCLE_COUNTER_SCOPE_DETAILED(SLATE_STATS_DETAIL_LEVEL_HI, GSlateMeasureStringTime);
 	FCharacterList& CharacterList = FontCache->GetCharacterList( InFontInfo, FontScale );
 	const uint16 MaxHeight = CharacterList.GetMaxHeight();
 
diff --git a/Engine/Source/Runtime/SlateCore/Private/Rendering/ElementBatcher.cpp b/Engine/Source/Runtime/SlateCore/Private/Rendering/ElementBatcher.cpp
--- a/Engine/Source/Runtime/SlateCore/Private/Rendering/ElementBatcher.cpp
+++ b/Engine/Source/Runtime/SlateCore/Private/Rendering/ElementBatcher.cpp
@@ -5,6 +5,16 @@
 #include "SlateRenderTransform.h"
 #include "Internationalization/Text.h"
 
+DECLARE_DWORD_COUNTER_STAT(TEXT("Num Layers"), STAT_SlateNumLayers, STATGROUP_Slate);
+DECLARE_DWORD_COUNTER_STAT(TEXT("Num Batches"), STAT_SlateNumBatches, STATGROUP_Slate);
+DECLARE_DWORD_COUNTER_STAT(TEXT("Num Vertices"), STAT_SlateVertexCount, STATGROUP_Slate);
+DECLARE_MEMORY_STAT(TEXT("Batch Vertex Memory"), STAT_SlateVertexBatchMemory, STATGROUP_SlateMemory);
+DECLARE_MEMORY_STAT(TEXT("Batch Index Memory"), STAT_SlateIndexBatchMemory, STATGROUP_SlateMemory);
+
+SLATE_DECLARE_CYCLE_COUNTER(GSlateAddElements, "Add Elements");
+SLATE_DECLARE_CYCLE_COUNTER(GSlateFindBatchTime, "FindElementForBatch");
+SLATE_DECLARE_CYCLE_COUNTER(GSlateFillBatchBuffers, "FillBatchBuffers");
+
 // Super-hacky way of storing the scissor rect so we don't have to change all the FSlateDrawElement APIs for this hacky support.
 SLATECORE_API TOptional<FShortRect> GSlateScissorRect;
 
@@ -60,12 +70,38 @@ FSlateElementBatcher::~FSlateElementBatcher()
 
 void FSlateElementBatcher::AddElements( const TArray<FSlateDrawElement>& DrawElements )
 {
-	SCOPE_CYCLE_COUNTER(STAT_SlateBuildElements);
+	SLATE_CYCLE_COUNTER_SCOPE(GSlateAddElements);
+	// This stuff is just for the counters. Could be scoped by an #ifdef if necessary.
+	static_assert(
+		FSlateDrawElement::EElementType::ET_Box == 0 &&
+		FSlateDrawElement::EElementType::ET_DebugQuad == 1 &&
+		FSlateDrawElement::EElementType::ET_Text == 2 &&
+		FSlateDrawElement::EElementType::ET_Spline == 3 &&
+		FSlateDrawElement::EElementType::ET_Line == 4 &&
+		FSlateDrawElement::EElementType::ET_Gradient == 5 &&
+		FSlateDrawElement::EElementType::ET_Viewport == 6 &&
+		FSlateDrawElement::EElementType::ET_Border == 7 &&
+		FSlateDrawElement::EElementType::ET_Custom == 8 &&
+		FSlateDrawElement::EElementType::ET_Count == 9, 
+		"If FSlateDrawElement::EElementType is modified, this array must be made to match.");
+
+	static FName ElementFNames[] = 
+	{
+		FName(TEXT("Box")),
+		FName(TEXT("DebugQuad")),
+		FName(TEXT("Text")),
+		FName(TEXT("Spline")),
+		FName(TEXT("Line")),
+		FName(TEXT("Gradient")),
+		FName(TEXT("Viewport")),
+		FName(TEXT("Border")),
+		FName(TEXT("Custom")),
+	};
+
 
 	for( int32 DrawElementIndex = 0; DrawElementIndex < DrawElements.Num(); ++DrawElementIndex )
 	{
 		const FSlateDrawElement& DrawElement = DrawElements[DrawElementIndex];
-
 		const FSlateRect& InClippingRect = DrawElement.GetClippingRect();
 	
 		// A zero or negatively sized clipping rect means the geometry will not be displayed
@@ -78,6 +114,8 @@ void FSlateElementBatcher::AddElements( const TArray<FSlateDrawElement>& DrawEle
 			// scissor rects are sort of a low level hack, so no one konws to clip against them. Instead we do it here to make sure the element is not actually rendered.
 			if (!bIsScissored)
 			{
+				// time just the adding of the element. The clipping stuff will be counted in exclusive time for the non-typed timer.
+				SLATE_CYCLE_COUNTER_SCOPE_CUSTOM_DETAILED(SLATE_STATS_DETAIL_LEVEL_MED, GSlateAddElements, ElementFNames[DrawElement.GetElementType()]);
 				// Determine what type of element to add
 				switch( DrawElement.GetElementType() )
 				{
@@ -178,10 +216,9 @@ FSlateRenderTransform GetBoxRenderTransform(const FSlateDrawElement& DrawElement
 	return Concatenate(RotationTransform, ElementRenderTransform);
 }
 
+
 void FSlateElementBatcher::AddBoxElement( const FSlateDrawElement& DrawElement )
 {
-	SCOPE_CYCLE_COUNTER(STAT_SlateBuildBoxes);
-
 	const FSlateRenderTransform& ElementRenderTransform = DrawElement.GetRenderTransform();
 	const FSlateRenderTransform RenderTransform = GetBoxRenderTransform(DrawElement);
 	//const FVector2D& InPosition = DrawElement.GetPosition();
@@ -470,11 +507,8 @@ void FSlateElementBatcher::AddBoxElement( const FSlateDrawElement& DrawElement )
 	}
 }
 
-
 void FSlateElementBatcher::AddTextElement(const FSlateDrawElement& DrawElement)
 {
-	SCOPE_CYCLE_COUNTER(STAT_SlateBuildText);
-
 	//const FVector2D& InPosition = DrawElement.GetPosition();
 	//const FVector2D& Size = DrawElement.GetSize();
 	//const FVector2D& LocalSize = DrawElement.GetLocalSize();
@@ -640,11 +674,8 @@ void FSlateElementBatcher::AddTextElement(const FSlateDrawElement& DrawElement)
 	}
 }
 
-
 void FSlateElementBatcher::AddGradientElement( const FSlateDrawElement& DrawElement )
 {
-	SCOPE_CYCLE_COUNTER(STAT_SlateBuildGradients);
-
 	const FSlateRenderTransform& RenderTransform = DrawElement.GetRenderTransform();
 	//const FVector2D& InPosition = DrawElement.GetPosition();
 	//const FVector2D& Size = DrawElement.GetSize();
@@ -770,8 +801,6 @@ void FSlateElementBatcher::AddGradientElement( const FSlateDrawElement& DrawElem
 
 void FSlateElementBatcher::AddSplineElement( const FSlateDrawElement& DrawElement )
 {
-	SCOPE_CYCLE_COUNTER(STAT_SlateBuildSplines);
-
 	const FSlateRenderTransform& RenderTransform = DrawElement.GetRenderTransform();
 	//const FVector2D& InPosition = DrawElement.GetPosition();
 	//const FVector2D& Size = DrawElement.GetSize();
@@ -916,8 +945,6 @@ static bool LineIntersect( const FVector2D& P1, const FVector2D& P2, const FVect
 
 void FSlateElementBatcher::AddLineElement( const FSlateDrawElement& DrawElement )
 {
-	SCOPE_CYCLE_COUNTER(STAT_SlateBuildLines);
-
 	const FSlateRenderTransform& RenderTransform = DrawElement.GetRenderTransform();
 	//const FVector2D& InPosition = DrawElement.GetPosition();
 	//const FVector2D& Size = DrawElement.GetSize();
@@ -1038,15 +1065,15 @@ void FSlateElementBatcher::AddLineElement( const FSlateDrawElement& DrawElement
 				FSlateVertex StartV1 = BatchVertices[IndexStart-1];
 				FSlateVertex StartV2 = BatchVertices[IndexStart-2];
 
-				StartV1.TexCoords.X = StartPosRenderSpace.X;
-				StartV1.TexCoords.Y = StartPosRenderSpace.Y;
-				StartV1.TexCoords.Z = IntersectCenterRenderSpace.X;
-				StartV1.TexCoords.W = IntersectCenterRenderSpace.Y;
+				StartV1.TexCoords[0] = StartPosRenderSpace.X;
+				StartV1.TexCoords[1] = StartPosRenderSpace.Y;
+				StartV1.TexCoords[2] = IntersectCenterRenderSpace.X;
+				StartV1.TexCoords[3] = IntersectCenterRenderSpace.Y;
 
-				StartV2.TexCoords.X = StartPosRenderSpace.X;
-				StartV2.TexCoords.Y = StartPosRenderSpace.Y;
-				StartV2.TexCoords.Z = IntersectCenterRenderSpace.X;
-				StartV2.TexCoords.W = IntersectCenterRenderSpace.Y;
+				StartV2.TexCoords[0] = StartPosRenderSpace.X;
+				StartV2.TexCoords[1] = StartPosRenderSpace.Y;
+				StartV2.TexCoords[2] = IntersectCenterRenderSpace.X;
+				StartV2.TexCoords[3] = IntersectCenterRenderSpace.Y;
 
 				IndexStart += 2;
 				BatchVertices.Add( StartV2 );
@@ -1094,8 +1121,6 @@ void FSlateElementBatcher::AddLineElement( const FSlateDrawElement& DrawElement
 
 void FSlateElementBatcher::AddViewportElement( const FSlateDrawElement& DrawElement )
 {
-	SCOPE_CYCLE_COUNTER(STAT_SlateBuildViewports);
-
 	const FSlateRenderTransform& RenderTransform = DrawElement.GetRenderTransform();
 	//const FVector2D& InPosition = DrawElement.GetPosition();
 	//const FVector2D& Size = DrawElement.GetSize();
@@ -1184,8 +1209,6 @@ void FSlateElementBatcher::AddViewportElement( const FSlateDrawElement& DrawElem
 
 void FSlateElementBatcher::AddBorderElement( const FSlateDrawElement& DrawElement )
 {
-	SCOPE_CYCLE_COUNTER(STAT_SlateBuildBorders);
-
 	const FSlateRenderTransform& RenderTransform = DrawElement.GetRenderTransform();
 	//const FVector2D& InPosition = DrawElement.GetPosition();
 	//const FVector2D& Size = DrawElement.GetSize();
@@ -1430,7 +1453,7 @@ FSlateElementBatch& FSlateElementBatcher::FindBatchForElement(
 	ESlateBatchDrawFlag::Type DrawFlags,
 	const TOptional<FShortRect>& ScissorRect)
 {
-	SCOPE_CYCLE_COUNTER(STAT_SlateFindBatchTime);
+	SLATE_CYCLE_COUNTER_SCOPE_DETAILED(SLATE_STATS_DETAIL_LEVEL_HI, GSlateFindBatchTime);
 
 	// See if the layer already exists.
 	TSet<FSlateElementBatch>* ElementBatches = LayerToElementBatches.Find( Layer );
@@ -1529,7 +1552,7 @@ void FSlateElementBatcher::AddIndices( TArray<SlateIndex>& OutIndices, FSlateEle
 
 void FSlateElementBatcher::FillBatchBuffers( FSlateWindowElementList& WindowElementList, bool& bRequiresStencilTest )
 {
-	SCOPE_CYCLE_COUNTER( STAT_SlateUpdateBufferGTTime );
+	SLATE_CYCLE_COUNTER_SCOPE(GSlateFillBatchBuffers);
 
 	TArray<FSlateRenderBatch>& OutRenderBatches = WindowElementList.RenderBatches;
 	TArray<FSlateVertex>& OutBatchedVertices = WindowElementList.BatchedVertices;
@@ -1541,9 +1564,7 @@ void FSlateElementBatcher::FillBatchBuffers( FSlateWindowElementList& WindowElem
 	LayerToElementBatches.KeySort( TLess<uint32>() );
 
 
-#if STATS
-	NumLayers = FMath::Max<uint32>(NumLayers, LayerToElementBatches.Num() );
-#endif
+	STAT(NumLayers = FMath::Max<uint32>(NumLayers, LayerToElementBatches.Num() ));
 
 	bRequiresStencilTest = false;
 	// For each element batch add its vertices and indices to the bulk lists.
@@ -1551,9 +1572,7 @@ void FSlateElementBatcher::FillBatchBuffers( FSlateWindowElementList& WindowElem
 	{
 		TSet<FSlateElementBatch>& ElementBatches = It.Value();
 
-#if STATS
-		NumBatches += ElementBatches.Num();
-#endif
+		STAT(NumBatches += ElementBatches.Num());
 		for( TSet<FSlateElementBatch>::TIterator BatchIt(ElementBatches); BatchIt; ++BatchIt )
 		{
 			FSlateElementBatch& ElementBatch = *BatchIt;
@@ -1568,9 +1587,11 @@ void FSlateElementBatcher::FillBatchBuffers( FSlateWindowElementList& WindowElem
 				}
 
 				TArray<FSlateVertex>& BatchVertices = BatchVertexArrays[ ElementBatch.VertexArrayIndex ];
+				// after this loop, we'll be done with the array, so put it back on the free list.
 				VertexArrayFreeList.Add( ElementBatch.VertexArrayIndex );
 
 				TArray<SlateIndex>& BatchIndices = BatchIndexArrays[ ElementBatch.IndexArrayIndex ];
+				// after this loop, we'll be done with the array, so put it back on the free list.
 				IndexArrayFreeList.Add( ElementBatch.IndexArrayIndex );
 
 				// We should have at least some vertices and indices in the batch or none at all
diff --git a/Engine/Source/Runtime/SlateCore/Private/Rendering/RenderingCommon.cpp b/Engine/Source/Runtime/SlateCore/Private/Rendering/RenderingCommon.cpp
--- a/Engine/Source/Runtime/SlateCore/Private/Rendering/RenderingCommon.cpp
+++ b/Engine/Source/Runtime/SlateCore/Private/Rendering/RenderingCommon.cpp
@@ -73,10 +73,14 @@ FSlateVertex::FSlateVertex()
 }
 
 FSlateVertex::FSlateVertex(const FSlateRenderTransform& RenderTransform, const FVector2D& InLocalPosition, const FVector2D& InTexCoord, const FVector2D& InTexCoord2, const FColor& InColor, const FSlateRotatedClipRectType& InClipRect )
-	: TexCoords( InTexCoord, InTexCoord2 )
-	, ClipRect( InClipRect )
+	: ClipRect( InClipRect )
 	, Color( InColor )
 {
+	TexCoords[0] = InTexCoord.X;
+	TexCoords[1] = InTexCoord.Y;
+	TexCoords[2] = InTexCoord2.X;
+	TexCoords[3] = InTexCoord2.Y;
+
 	const FVector2D WindowPosition = TransformPoint(RenderTransform, InLocalPosition);
 	// Pixel snapping here.
 	Position[0] = FMath::RoundToInt(WindowPosition.X);
@@ -84,10 +88,14 @@ FSlateVertex::FSlateVertex(const FSlateRenderTransform& RenderTransform, const F
 }
 
 FSlateVertex::FSlateVertex( const FSlateRenderTransform& RenderTransform, const FVector2D& InLocalPosition, const FVector2D& InTexCoord, const FColor& InColor, const FSlateRotatedClipRectType& InClipRect )
-	: TexCoords( InTexCoord, FVector2D(1.0f, 1.0f) )
-	, ClipRect( InClipRect )
+	: ClipRect( InClipRect )
 	, Color( InColor )
 {
+	TexCoords[0] = InTexCoord.X;
+	TexCoords[1] = InTexCoord.Y;
+	TexCoords[2] = 1.0f;
+	TexCoords[3] = 1.0f;
+
 	const FVector2D WindowPosition = TransformPoint(RenderTransform, InLocalPosition);
 	// Pixel snapping here.
 	Position[0] = FMath::RoundToInt(WindowPosition.X);
diff --git a/Engine/Source/Runtime/SlateCore/Private/SlateCoreModule.cpp b/Engine/Source/Runtime/SlateCore/Private/SlateCoreModule.cpp
--- a/Engine/Source/Runtime/SlateCore/Private/SlateCoreModule.cpp
+++ b/Engine/Source/Runtime/SlateCore/Private/SlateCoreModule.cpp
@@ -7,44 +7,6 @@
 DEFINE_LOG_CATEGORY(LogSlate);
 DEFINE_LOG_CATEGORY(LogSlateStyles);
 
-DEFINE_STAT(STAT_SlateMeasureStringTime);
-DEFINE_STAT(STAT_SlateRenderingRTTime);
-DEFINE_STAT(STAT_SlatePresentRTTime);
-DEFINE_STAT(STAT_SlateUpdateBufferRTTime);
-DEFINE_STAT(STAT_SlateDrawTime);
-DEFINE_STAT(STAT_SlateUpdateBufferGTTime);
-DEFINE_STAT(STAT_SlateBuildLines);
-DEFINE_STAT(STAT_SlateBuildBoxes);
-DEFINE_STAT(STAT_SlateBuildBorders);
-DEFINE_STAT(STAT_SlateBuildText);
-DEFINE_STAT(STAT_SlateBuildGradients);
-DEFINE_STAT(STAT_SlateBuildSplines);
-DEFINE_STAT(STAT_SlateBuildElements);
-DEFINE_STAT(STAT_SlateFindBatchTime);
-DEFINE_STAT(STAT_SlateFontCachingTime);
-DEFINE_STAT(STAT_SlateRenderingGTTime);
-DEFINE_STAT(STAT_SlateBuildViewports);
-
-DEFINE_STAT(STAT_SlateNumTextureAtlases);
-DEFINE_STAT(STAT_SlateNumFontAtlases);
-DEFINE_STAT(STAT_SlateNumNonAtlasedTextures);
-DEFINE_STAT(STAT_SlateNumDynamicTextures);
-
-DEFINE_STAT(STAT_SlateNumLayers);
-DEFINE_STAT(STAT_SlateNumBatches);
-DEFINE_STAT(STAT_SlateVertexCount);
-
-DEFINE_STAT(STAT_SlateTextureDataMemory);
-DEFINE_STAT(STAT_SlateTextureAtlasMemory);
-DEFINE_STAT(STAT_SlateTextureGPUMemory);
-DEFINE_STAT(STAT_SlateFontKerningTableMemory);
-DEFINE_STAT(STAT_SlateFontMeasureCacheMemory);
-DEFINE_STAT(STAT_SlateVertexBatchMemory);
-DEFINE_STAT(STAT_SlateIndexBatchMemory);
-DEFINE_STAT(STAT_SlateVertexBufferMemory);
-DEFINE_STAT(STAT_SlateIndexBufferMemory);
-
-
 const float FOptionalSize::Unspecified = -1.0f;
 
 
diff --git a/Engine/Source/Runtime/SlateCore/Private/Stats/SlateStats.cpp b/Engine/Source/Runtime/SlateCore/Private/Stats/SlateStats.cpp
new file mode 100644
--- /dev/null
+++ b/Engine/Source/Runtime/SlateCore/Private/Stats/SlateStats.cpp
@@ -0,0 +1,233 @@
+// Copyright 1998-2014 Epic Games, Inc. All Rights Reserved.
+
+#include "SlateCorePrivatePCH.h"
+
+#if SLATE_STATS
+
+SLATE_DECLARE_CYCLE_COUNTER(GSlateStatsOverhead, "Stats Overhead");
+
+SLATECORE_API int32 GSlateStatsFlatEnable = 0;
+FAutoConsoleVariableRef GSlateStatsFlatEnableCVar(
+	TEXT("Slate.Stats.Flat.Enable"), 
+	GSlateStatsFlatEnable, 
+	TEXT("Set to > 0 to enable slate flat stats capture. Will update averages every Slate.Stats.Flat.IntervalWindowSec seconds."));
+
+SLATECORE_API int32 GSlateStatsFlatLogOutput = 0;
+FAutoConsoleVariableRef GSlateStatsFlatLogOutputCVar(
+	TEXT("Slate.Stats.Flat.LogOutput"), 
+	GSlateStatsFlatLogOutput, 
+	TEXT("Set to > 0 to log the stat flat stats every time a new window average is calculated."));
+
+SLATECORE_API float GSlateStatsFlatIntervalWindowSec = 1.0f;
+FAutoConsoleVariableRef GSlateStatsFlatIntervalWindowSecCVar(
+	TEXT("Slate.Stats.Flat.IntervalWindowSec"), 
+	GSlateStatsFlatIntervalWindowSec, 
+	TEXT("Interval window (in seconds) to update the slate stat averages. Slate.Stats.Flat.Enable must be > 0 also."));
+
+SLATECORE_API int32 GSlateStatsHierarchyTrigger = 0;
+FAutoConsoleVariableRef GSlateStatsHierarchyTriggerCVar(
+	TEXT("Slate.Stats.Hierarchy.Trigger"), 
+	GSlateStatsHierarchyTrigger, 
+	TEXT("Set to > 0 to capture a slate hierarchical profile on the next frame. Will be reset after that frame."));
+
+// another shortcut for using a single-threaded solution. tracks when we update our averages so external observers can poll a quick value.
+bool GAverageInclusiveTimesWereUpdatedThisFrame = false;
+
+/** Slate hierarchical stats are not resizeable to keep the system fast. The system will detect the need to resize and assert in debug builds, but will not in optimized builds to reduce stat overhead. */
+#define SLATE_STATS_HIERARCHY_MAX_ENTRIES 300000
+
+FSlateStatCycleCounter::FSlateStatCycleCounter(FName InCounterName) : Name(InCounterName)
+, InclusiveTime(0.0)
+, LastComputedAverageInclusiveTime(0.0)
+, StackDepth(0)
+, StartTime(0.0)
+{
+	GetGlobalRegisteredCounters().Add(this);
+}
+
+FSlateStatCycleCounter::~FSlateStatCycleCounter()
+{
+	GetGlobalRegisteredCounters().Remove(this);
+}
+
+void FSlateStatCycleCounter::Reset()
+{
+	InclusiveTime = 0.0;
+	StackDepth = 0;
+	// Don't clear LastComputedAverageInclusiveTime because we want people to be able to 
+	// report on a stat with a stable number while we accumulate a new average.
+	// Also don't clear StartTime because we chould be in the middle of a timing (like the self-profiling timer every frame).
+}
+
+const TArray<FSlateStatCycleCounter*>& FSlateStatCycleCounter::GetRegisteredCounters()
+{
+	return GetGlobalRegisteredCounters();
+}
+
+TArray<FSlateStatCycleCounter*>& FSlateStatCycleCounter::GetGlobalRegisteredCounters()
+{
+	static TArray<FSlateStatCycleCounter*>* GlobalRegisteredCounters = new TArray<FSlateStatCycleCounter*>();
+	return *GlobalRegisteredCounters;
+}
+
+void FSlateStatCycleCounter::EndFrame(double CurrentTime)
+{
+	// Ensure the overhead is tracked at minimum even if all other detail levels are turned off.
+	const int32 SLATE_STATS_DETAIL_LEVEL_FORCE_ON = INT_MIN;
+	SLATE_CYCLE_COUNTER_SCOPE_FLAT_DETAILED(SLATE_STATS_DETAIL_LEVEL_FORCE_ON, GSlateStatsOverhead);
+	
+	// dump stats hierarchy if one was collected this frame.
+	if (FSlateStatHierarchy::Get().GetStatEntries().Num() > 0)
+	{
+		// Place in the <UE4>\<GAME>\Saved\<InFolderName> folder
+		FString Filename = FString::Printf(TEXT("%sSlateHierachyStats-%s.csv"), *FPaths::GameSavedDir(), *FDateTime::Now().ToString());
+		UE_LOG(LogSlate, Log, TEXT("Dumping Slate Hierarchy Stats to %s..."), *Filename);
+		FArchive* OutputStream = IFileManager::Get().CreateFileWriter(*Filename, EFileWrite::FILEWRITE_NoReplaceExisting);
+
+		// some state vars used to print the path of the stat.
+		const int kMaxDepth = 100;
+		int32 Path[kMaxDepth] = { 0 };
+		int32 PathCurrentDepth = -1;
+		TCHAR PathStr[4 * kMaxDepth] = { 0 };
+		TCHAR* PathStrCurrent = PathStr;
+		FSlateStatHierarchy::Get().ComputeExclusiveTimes();
+
+		for (const auto& Entry : FSlateStatHierarchy::Get().GetStatEntries())
+		{
+			const bool bDescending = Entry.StackDepth > PathCurrentDepth;
+			const bool bAscending = Entry.StackDepth < PathCurrentDepth;
+			const bool bOverflow = PathCurrentDepth >= (int32)ARRAY_COUNT(Path);
+			const int PathPrevDepth = PathCurrentDepth;
+			PathCurrentDepth = Entry.StackDepth;
+
+			if (!bOverflow)
+			{
+				if (bDescending)
+				{
+					// we always increment the ordinal, so init to -1 so the first value is zero.
+					Path[PathCurrentDepth] = -1;
+					// put a dot after the current depth and track the new position to put new depths.
+					if (PathCurrentDepth > 0)
+					{
+						while (*PathStrCurrent != NULL) ++PathStrCurrent;
+						*PathStrCurrent++ = '.';
+						*PathStrCurrent = 0;
+					}
+				}
+				else if (bAscending)
+				{
+					// back up until we find the dot matching our current depth (or get to depth zero)
+					if (PathCurrentDepth == 0)
+					{
+						// we are at level zero, so just back up all the way. Saves us from checking boundaries in the loop below when we are not at zero depth.
+						PathStrCurrent = PathStr;
+					}
+					else
+					{
+						int LevelsAscended = 0;
+						const int LevelsToAscend = PathPrevDepth - PathCurrentDepth;
+						while (LevelsAscended < LevelsToAscend)
+						{
+							// we are not back to level zero, so we know we will find a '.' eventually.
+							PathStrCurrent -= 2; // back us up before the dot marking this level
+							// keep backing up until we find the dot for the previous level.
+							while (*(PathStrCurrent - 1) != '.')
+							{
+								--PathStrCurrent;
+							}
+							++LevelsAscended;
+						}
+					}
+				}
+				// increment the ordinal at this depth
+				++Path[PathCurrentDepth];
+				FCString::Sprintf(PathStrCurrent, TEXT("%4d"), Path[PathCurrentDepth]);
+			}
+
+			OutputStream->Logf(TEXT("%s,%s,%s,%.8f,%.8f"), PathStr, *Entry.CounterName.ToString(), Entry.CustomName != NAME_None ? *Entry.CustomName.ToString() : TEXT(""), Entry.InclusiveTime*1000, Entry.ExclusiveTime*1000);
+			//UE_LOG(LogSlate, Log, TEXT("%s,%s,%s,%.8f,%.8f"), PathStr, *Entry.CounterName.ToString(), Entry.CustomName != NAME_None ? *Entry.CustomName.ToString() : TEXT(""), Entry.InclusiveTime*1000, Entry.ExclusiveTime*1000);
+		}
+		OutputStream->Close();
+		delete OutputStream; OutputStream = nullptr;
+		UE_LOG(LogSlate, Log, TEXT("Done dumping Slate Hierarchy Stats!"), *Filename);
+	}
+
+	// static state for tracking when to re-average the stats 
+	static double Time = CurrentTime;
+	static double LastTime = CurrentTime;
+	static double NumFrames = 0.0;
+	NumFrames += 1.0;
+	Time = CurrentTime;
+	const double Delta = Time - LastTime;
+	// let things settle down before taking the first real sample
+	static double NextDelta = GSlateStatsFlatIntervalWindowSec;
+	// ensure this gets reset every frame.
+	GAverageInclusiveTimesWereUpdatedThisFrame = false;
+
+	// output flat stats if it's time to do so
+	if (Delta > NextDelta)
+	{
+		LastTime = Time;
+		NumFrames = NumFrames / 1000.0; // convert to ms.
+		if (GSlateStatsFlatLogOutput > 0)
+		{
+			UE_LOG(LogSlate, Log, TEXT("Slate Flat Stats"));
+			UE_LOG(LogSlate, Log, TEXT("================"));
+		}
+
+		// iterate over the counters, outputting their data and resetting them.
+		for (const auto Counter : FSlateStatCycleCounter::GetRegisteredCounters())
+		{
+			Counter->LastComputedAverageInclusiveTime = Counter->InclusiveTime / NumFrames;
+			if (GSlateStatsFlatLogOutput > 0)
+			{
+				UE_LOG(LogSlate, Log, TEXT("%s,%.8f"), *Counter->GetName().ToString(), Counter->LastComputedAverageInclusiveTime);
+			}
+			Counter->Reset();
+		}
+		// Frame time is a "virtual stat" so output it like a regular stat.
+		if (GSlateStatsFlatLogOutput > 0)
+		{
+			UE_LOG(LogSlate, Log, TEXT("%s,%.8f"), TEXT("Frame Time"), Delta / NumFrames);
+		}
+
+		LastTime = Time;
+		NumFrames = 0.0;
+		NextDelta = GSlateStatsFlatIntervalWindowSec;
+		GAverageInclusiveTimesWereUpdatedThisFrame = true;
+	}
+
+	// Clear the hierarchy entries, and tell the system to possibly track hierarchy entries next frame.
+	FSlateStatHierarchy::Get().EndFrame(GSlateStatsHierarchyTrigger > 0);
+	GSlateStatsHierarchyTrigger = 0;
+}
+
+bool FSlateStatCycleCounter::AverageInclusiveTimesWereUpdatedThisFrame()
+{
+	return GAverageInclusiveTimesWereUpdatedThisFrame;
+}
+
+FSlateStatHierarchy& FSlateStatHierarchy::Get()
+{
+	static FSlateStatHierarchy Singleton;
+	return Singleton;
+}
+
+FSlateStatHierarchy::FSlateStatHierarchy() 
+	: StackDepth(0)
+	, bTrackThisFrame(false)
+{
+}
+
+void FSlateStatHierarchy::EndFrame(bool bTrackNextFrame)
+{
+	// don't allocate space for hierarchical profiling unless we want the next frame to be captured.
+	// Otherwise, try to maintain the same amount of reserved space.
+	int32 RequiredSlack = bTrackNextFrame ? SLATE_STATS_HIERARCHY_MAX_ENTRIES : StatEntries.Max();
+	// !!! WRH 2014/12/02 - This should be a block-based list containers so resizes are allowed without reallocating existing entries.
+	StatEntries.Empty(RequiredSlack);
+	bTrackThisFrame = bTrackNextFrame;
+	StackDepth = 0;
+}
+
+#endif
diff --git a/Engine/Source/Runtime/SlateCore/Private/Textures/TextureAtlas.cpp b/Engine/Source/Runtime/SlateCore/Private/Textures/TextureAtlas.cpp
--- a/Engine/Source/Runtime/SlateCore/Private/Textures/TextureAtlas.cpp
+++ b/Engine/Source/Runtime/SlateCore/Private/Textures/TextureAtlas.cpp
@@ -3,6 +3,9 @@
 #include "SlateCorePrivatePCH.h"
 #include "SlateRenderer.h"
 
+DEFINE_STAT(STAT_SlateTextureGPUMemory);
+DEFINE_STAT(STAT_SlateTextureDataMemory);
+DECLARE_MEMORY_STAT(TEXT("Texture Atlas CPU Memory"), STAT_SlateTextureAtlasMemory, STATGROUP_SlateMemory);
 
 /* FSlateTextureAtlas helper class
  *****************************************************************************/
diff --git a/Engine/Source/Runtime/SlateCore/Private/Widgets/SCompoundWidget.cpp b/Engine/Source/Runtime/SlateCore/Private/Widgets/SCompoundWidget.cpp
--- a/Engine/Source/Runtime/SlateCore/Private/Widgets/SCompoundWidget.cpp
+++ b/Engine/Source/Runtime/SlateCore/Private/Widgets/SCompoundWidget.cpp
@@ -2,16 +2,11 @@
 
 #include "SlateCorePrivatePCH.h"
 
-DECLARE_CYCLE_STAT( TEXT("OnPaint SCompoundWidget"), STAT_SlateOnPaint_SCompoundWidget, STATGROUP_Slate );
-
 int32 SCompoundWidget::OnPaint( const FPaintArgs& Args, const FGeometry& AllottedGeometry, const FSlateRect& MyClippingRect, FSlateWindowElementList& OutDrawElements, int32 LayerId, const FWidgetStyle& InWidgetStyle, bool bParentEnabled ) const
 {
 	// A CompoundWidget just draws its children
 	FArrangedChildren ArrangedChildren(EVisibility::Visible);
 	{
-#if SLATE_HD_STATS
-		SCOPE_CYCLE_COUNTER( STAT_SlateOnPaint_SCompoundWidget );
-#endif
 		this->ArrangeChildren(AllottedGeometry, ArrangedChildren);
 	}
 
diff --git a/Engine/Source/Runtime/SlateCore/Private/Widgets/SOverlay.cpp b/Engine/Source/Runtime/SlateCore/Private/Widgets/SOverlay.cpp
--- a/Engine/Source/Runtime/SlateCore/Private/Widgets/SOverlay.cpp
+++ b/Engine/Source/Runtime/SlateCore/Private/Widgets/SOverlay.cpp
@@ -4,8 +4,6 @@
 
 #include "LayoutUtils.h"
 
-DECLARE_CYCLE_STAT( TEXT("OnPaint SOverlay"), STAT_SlateOnPaint_SOverlay, STATGROUP_Slate );
-
 SOverlay::SOverlay()
 : Children()
 {
@@ -82,9 +80,6 @@ int32 SOverlay::OnPaint( const FPaintArgs& Args, const FGeometry& AllottedGeomet
 {
 	FArrangedChildren ArrangedChildren(EVisibility::Visible);
 	{
-#if SLATE_HD_STATS
-		SCOPE_CYCLE_COUNTER( STAT_SlateOnPaint_SOverlay );
-#endif
 		// The box panel has no visualization of its own; it just visualizes its children.
 		this->ArrangeChildren(AllottedGeometry, ArrangedChildren);
 	}
diff --git a/Engine/Source/Runtime/SlateCore/Private/Widgets/SPanel.cpp b/Engine/Source/Runtime/SlateCore/Private/Widgets/SPanel.cpp
--- a/Engine/Source/Runtime/SlateCore/Private/Widgets/SPanel.cpp
+++ b/Engine/Source/Runtime/SlateCore/Private/Widgets/SPanel.cpp
@@ -2,8 +2,6 @@
 
 #include "SlateCorePrivatePCH.h"
 
-DECLARE_CYCLE_STAT( TEXT("OnPaint SPanel"), STAT_SlateOnPaint_SPanel, STATGROUP_Slate );
-
 /* SWidget overrides
  *****************************************************************************/
 
@@ -13,9 +11,6 @@ int32 SPanel::OnPaint( const FPaintArgs& Args, const FGeometry& AllottedGeometry
 	// OutDrawElements.AddItem( FSlateDrawElement::MakeQuad( AllottedGeometry.AbsolutePosition, AllottedGeometry.Size ) );	
 	FArrangedChildren ArrangedChildren(EVisibility::Visible);
 	{
-#if SLATE_HD_STATS
-		SCOPE_CYCLE_COUNTER( STAT_SlateOnPaint_SPanel );
-#endif
 		// The box panel has no visualization of its own; it just visualizes its children.
 		
 		this->ArrangeChildren(AllottedGeometry, ArrangedChildren);
diff --git a/Engine/Source/Runtime/SlateCore/Private/Widgets/SWidget.cpp b/Engine/Source/Runtime/SlateCore/Private/Widgets/SWidget.cpp
--- a/Engine/Source/Runtime/SlateCore/Private/Widgets/SWidget.cpp
+++ b/Engine/Source/Runtime/SlateCore/Private/Widgets/SWidget.cpp
@@ -4,12 +4,19 @@
 #include "Widgets/SWidget.h"
 #include "Input/Events.h"
 #include "ActiveTimerHandle.h"
+#include "SlateStats.h"
 
-DECLARE_CYCLE_STAT(TEXT("OnPaint"), STAT_SlateOnPaint, STATGROUP_Slate);
-DECLARE_CYCLE_STAT(TEXT("ArrangeChildren"), STAT_SlateArrangeChildren, STATGROUP_Slate);
 DECLARE_DWORD_COUNTER_STAT(TEXT("Num Painted Widgets"), STAT_SlateNumPaintedWidgets, STATGROUP_Slate);
 DECLARE_DWORD_COUNTER_STAT(TEXT("Num Ticked Widgets"), STAT_SlateNumTickedWidgets, STATGROUP_Slate);
 
+SLATE_DECLARE_CYCLE_COUNTER(GSlateWidgetTick, "SWidget Tick");
+SLATE_DECLARE_CYCLE_COUNTER(GSlateOnPaint, "OnPaint");
+SLATE_DECLARE_CYCLE_COUNTER(GSlatePrepass, "SlatePrepass");
+SLATE_DECLARE_CYCLE_COUNTER(GSlateArrangeChildren, "ArrangeChildren");
+SLATE_DECLARE_CYCLE_COUNTER(GSlateGetVisibility, "GetVisibility");
+
+TAutoConsoleVariable<int32> TickInvisibleWidgets(TEXT("Slate.TickInvisibleWidgets"), 0, TEXT("Controls whether invisible widgets are ticked."));
+
 SWidget::SWidget()
 	: CreatedInFile( TEXT("") )
 	, CreatedOnLine( -1 )
@@ -342,15 +349,16 @@ void SWidget::Tick( const FGeometry& AllottedGeometry, const double InCurrentTim
 {
 }
 
-TAutoConsoleVariable<int32> TickInvisibleWidgets(TEXT("Slate.TickInvisibleWidgets"), 0, TEXT("Controls whether invisible widgets are ticked."));
-
 void SWidget::TickWidgetsRecursively( const FGeometry& AllottedGeometry, const double InCurrentTime, const float InDeltaTime )
 {
 	INC_DWORD_STAT(STAT_SlateNumTickedWidgets);
 
 	// Execute any pending active timers for this widget, followed by the passive tick
 	ExecuteActiveTimers( InCurrentTime, InDeltaTime );
-	Tick( AllottedGeometry, InCurrentTime, InDeltaTime );
+	{
+		SLATE_CYCLE_COUNTER_SCOPE_CUSTOM_DETAILED(SLATE_STATS_DETAIL_LEVEL_MED, GSlateWidgetTick, GetType());
+		Tick(AllottedGeometry, InCurrentTime, InDeltaTime);
+	}
 
 	// Gather all children, whether they're visible or not.  We need to allow invisible widgets to
 	// consider whether they should still be invisible in their tick functions, as well as maintain
@@ -374,6 +382,7 @@ void SWidget::SlatePrepass()
 
 void SWidget::SlatePrepass(float LayoutScaleMultiplier)
 {
+	SLATE_CYCLE_COUNTER_SCOPE_CUSTOM_DETAILED(SLATE_STATS_DETAIL_LEVEL_MED, GSlatePrepass, GetType());
 	// Cache child desired sizes first. This widget's desired size is
 	// a function of its children's sizes.
 	FChildren* MyChildren = this->GetChildren();
@@ -553,6 +562,10 @@ FString SWidget::GetCreatedInFile() const
 	return this->CreatedInFileFullPath.ToString();
 }
 
+FName SWidget::GetCreatedInFileFName() const
+{
+	return this->CreatedInFile;
+}
 
 int32 SWidget::GetCreatedInLineNumber() const
 {
@@ -622,17 +635,19 @@ void SWidget::SetDebugInfo( const ANSICHAR* InType, const ANSICHAR* InFile, int3
 {
 	this->TypeOfWidget = InType;
 	this->CreatedInFileFullPath = FName( InFile );
+	this->CreatedInFileFullPath.SetNumber(OnLine);
 	this->CreatedInFile = FName( *FPaths::GetCleanFilename(InFile) );
+	this->CreatedInFile.SetNumber(OnLine);
 	this->CreatedOnLine = OnLine;
 }
 
-SLATECORE_API int32 bFoldTick = 1;
+SLATECORE_API int32 bFoldTick = 0;
 FAutoConsoleVariableRef FoldTick(TEXT("Slate.FoldTick"), bFoldTick, TEXT("When folding, call Tick as part of the paint pass instead of a separate tick pass."));
 
 int32 SWidget::Paint(const FPaintArgs& Args, const FGeometry& AllottedGeometry, const FSlateRect& MyClippingRect, FSlateWindowElementList& OutDrawElements, int32 LayerId, const FWidgetStyle& InWidgetStyle, bool bParentEnabled) const
 {
-	SCOPE_CYCLE_COUNTER(STAT_SlateOnPaint);
 	INC_DWORD_STAT(STAT_SlateNumPaintedWidgets);
+	SLATE_CYCLE_COUNTER_SCOPE_CUSTOM_DETAILED(SLATE_STATS_DETAIL_LEVEL_MED, GSlateOnPaint, GetType());
 
 	if ( bFoldTick )
 	{
@@ -679,10 +694,16 @@ float SWidget::GetRelativeLayoutScale(const FSlotBase& Child) const
 
 void SWidget::ArrangeChildren(const FGeometry& AllottedGeometry, FArrangedChildren& ArrangedChildren) const
 {
-	SCOPE_CYCLE_COUNTER(STAT_SlateArrangeChildren);
+	SLATE_CYCLE_COUNTER_SCOPE_CUSTOM_DETAILED(SLATE_STATS_DETAIL_LEVEL_MED, GSlateArrangeChildren, GetType());
 	OnArrangeChildren(AllottedGeometry, ArrangedChildren);
 }
 
+EVisibility SWidget::GetVisibility() const
+{
+	SLATE_CYCLE_COUNTER_SCOPE_CUSTOM_DETAILED(SLATE_STATS_DETAIL_LEVEL_HI,GSlateGetVisibility, GetType());
+	return Visibility.Get();
+}
+
 TSharedRef<FActiveTimerHandle> SWidget::RegisterActiveTimer(float TickPeriod, FWidgetActiveTimerDelegate TickFunction)
 {
 	TSharedRef<FActiveTimerHandle> ActiveTimerHandle = MakeShareable(new FActiveTimerHandle(TickPeriod, TickFunction, FSlateApplicationBase::Get().GetCurrentTime() + TickPeriod));
diff --git a/Engine/Source/Runtime/SlateCore/Public/Rendering/DrawElements.h b/Engine/Source/Runtime/SlateCore/Public/Rendering/DrawElements.h
--- a/Engine/Source/Runtime/SlateCore/Public/Rendering/DrawElements.h
+++ b/Engine/Source/Runtime/SlateCore/Public/Rendering/DrawElements.h
@@ -159,6 +159,7 @@ public:
 		ET_Viewport,
 		ET_Border,
 		ET_Custom,
+		ET_Count,
 	};
 
 	enum ERotationSpace
diff --git a/Engine/Source/Runtime/SlateCore/Public/Rendering/ElementBatcher.h b/Engine/Source/Runtime/SlateCore/Public/Rendering/ElementBatcher.h
--- a/Engine/Source/Runtime/SlateCore/Public/Rendering/ElementBatcher.h
+++ b/Engine/Source/Runtime/SlateCore/Public/Rendering/ElementBatcher.h
@@ -122,9 +122,6 @@ private:
 	void AddIndices( TArray<SlateIndex>& OutIndices, FSlateElementBatch& ElementBatch, const TArray<SlateIndex>& IndexBatch );
 
 private:
-
-	const FSlateBrush *SplineBrush;
-
 	// Element batch maps sorted by layer.
 	TMap<uint32, TSet<FSlateElementBatch>> LayerToElementBatches;
 
diff --git a/Engine/Source/Runtime/SlateCore/Public/Rendering/RenderingCommon.h b/Engine/Source/Runtime/SlateCore/Public/Rendering/RenderingCommon.h
--- a/Engine/Source/Runtime/SlateCore/Public/Rendering/RenderingCommon.h
+++ b/Engine/Source/Runtime/SlateCore/Public/Rendering/RenderingCommon.h
@@ -166,7 +166,7 @@ typedef FSlateRotatedRect FSlateRotatedClipRectType;
 struct FSlateVertex
 {
 	/** Texture coordinates.  The first 2 are in xy and the 2nd are in zw */
-	FVector4 TexCoords; 
+	float TexCoords[4]; 
 	/** Position of the vertex in window space */
 	int16 Position[2];
 	/** clip center/extents in render window space (window space with render transforms applied) */
diff --git a/Engine/Source/Runtime/SlateCore/Public/SlateCore.h b/Engine/Source/Runtime/SlateCore/Public/SlateCore.h
--- a/Engine/Source/Runtime/SlateCore/Public/SlateCore.h
+++ b/Engine/Source/Runtime/SlateCore/Public/SlateCore.h
@@ -17,44 +17,6 @@
 SLATECORE_API DECLARE_LOG_CATEGORY_EXTERN(LogSlate, Log, All);
 SLATECORE_API DECLARE_LOG_CATEGORY_EXTERN(LogSlateStyles, Log, All);
 
-DECLARE_CYCLE_STAT_EXTERN(TEXT("Measure String"), STAT_SlateMeasureStringTime, STATGROUP_Slate,);
-DECLARE_CYCLE_STAT_EXTERN(TEXT("Slate Rendering RT Time"), STAT_SlateRenderingRTTime, STATGROUP_Slate , SLATECORE_API);
-DECLARE_CYCLE_STAT_EXTERN(TEXT("Slate RT Present Time"), STAT_SlatePresentRTTime, STATGROUP_Slate , SLATECORE_API);
-DECLARE_CYCLE_STAT_EXTERN(TEXT("Update Buffers RT"), STAT_SlateUpdateBufferRTTime, STATGROUP_Slate , SLATECORE_API);
-DECLARE_CYCLE_STAT_EXTERN(TEXT("Draw Time"), STAT_SlateDrawTime, STATGROUP_Slate, SLATECORE_API);
-DECLARE_CYCLE_STAT_EXTERN(TEXT("Update Buffers GT"), STAT_SlateUpdateBufferGTTime, STATGROUP_Slate,);
-DECLARE_CYCLE_STAT_EXTERN(TEXT("Font Caching"), STAT_SlateFontCachingTime, STATGROUP_Slate,);
-DECLARE_CYCLE_STAT_EXTERN(TEXT("Build Elements"), STAT_SlateBuildElements, STATGROUP_Slate,);
-DECLARE_CYCLE_STAT_EXTERN(TEXT("Find Batch"), STAT_SlateFindBatchTime, STATGROUP_Slate,);
-DECLARE_CYCLE_STAT_EXTERN(TEXT("Build Lines"), STAT_SlateBuildLines, STATGROUP_Slate,);
-DECLARE_CYCLE_STAT_EXTERN(TEXT("Build Splines"), STAT_SlateBuildSplines, STATGROUP_Slate,);
-DECLARE_CYCLE_STAT_EXTERN(TEXT("Build Gradients"), STAT_SlateBuildGradients, STATGROUP_Slate,);
-DECLARE_CYCLE_STAT_EXTERN(TEXT("Build Text"), STAT_SlateBuildText, STATGROUP_Slate,);
-DECLARE_CYCLE_STAT_EXTERN(TEXT("Build Border"), STAT_SlateBuildBorders, STATGROUP_Slate,);
-DECLARE_CYCLE_STAT_EXTERN(TEXT("Build Boxes"), STAT_SlateBuildBoxes, STATGROUP_Slate,);
-DECLARE_CYCLE_STAT_EXTERN(TEXT("Slate Rendering GT Time"), STAT_SlateRenderingGTTime, STATGROUP_Slate, SLATECORE_API);
-DECLARE_CYCLE_STAT_EXTERN(TEXT("Build Viewports"), STAT_SlateBuildViewports, STATGROUP_Slate,);
-
-DECLARE_DWORD_ACCUMULATOR_STAT_EXTERN(TEXT("Num Texture Atlases"), STAT_SlateNumTextureAtlases, STATGROUP_SlateMemory, SLATECORE_API);
-DECLARE_DWORD_ACCUMULATOR_STAT_EXTERN(TEXT("Num Font Atlases"), STAT_SlateNumFontAtlases, STATGROUP_SlateMemory,);
-DECLARE_DWORD_ACCUMULATOR_STAT_EXTERN(TEXT("Num Non-Atlased Textures"), STAT_SlateNumNonAtlasedTextures, STATGROUP_SlateMemory,SLATECORE_API);
-DECLARE_DWORD_ACCUMULATOR_STAT_EXTERN(TEXT("Num Dynamic Textures"), STAT_SlateNumDynamicTextures, STATGROUP_SlateMemory, SLATECORE_API);
-
-DECLARE_DWORD_COUNTER_STAT_EXTERN(TEXT("Num Layers"), STAT_SlateNumLayers, STATGROUP_Slate,);
-DECLARE_DWORD_COUNTER_STAT_EXTERN(TEXT("Num Batches"), STAT_SlateNumBatches, STATGROUP_Slate,);
-DECLARE_DWORD_COUNTER_STAT_EXTERN(TEXT("Num Vertices"), STAT_SlateVertexCount, STATGROUP_Slate,);
-
-DECLARE_MEMORY_STAT_EXTERN(TEXT("Texture Data CPU Memory"), STAT_SlateTextureDataMemory, STATGROUP_SlateMemory, SLATECORE_API);
-DECLARE_MEMORY_STAT_EXTERN(TEXT("Texture Atlas CPU Memory"), STAT_SlateTextureAtlasMemory, STATGROUP_SlateMemory,);
-DECLARE_MEMORY_STAT_EXTERN(TEXT("Texture Data GPU Memory"), STAT_SlateTextureGPUMemory, STATGROUP_SlateMemory, SLATECORE_API);
-DECLARE_MEMORY_STAT_EXTERN(TEXT("Font Kerning Table Memory"), STAT_SlateFontKerningTableMemory, STATGROUP_SlateMemory,);
-DECLARE_MEMORY_STAT_EXTERN(TEXT("Font Measure Memory"), STAT_SlateFontMeasureCacheMemory, STATGROUP_SlateMemory,);
-DECLARE_MEMORY_STAT_EXTERN(TEXT("Batch Vertex Memory"), STAT_SlateVertexBatchMemory, STATGROUP_SlateMemory,);
-DECLARE_MEMORY_STAT_EXTERN(TEXT("Batch Index Memory"), STAT_SlateIndexBatchMemory, STATGROUP_SlateMemory,);
-DECLARE_MEMORY_STAT_EXTERN(TEXT("Vertex Buffer Memory"), STAT_SlateVertexBufferMemory, STATGROUP_SlateMemory, SLATECORE_API);
-DECLARE_MEMORY_STAT_EXTERN(TEXT("Index Buffer Memory"), STAT_SlateIndexBufferMemory, STATGROUP_SlateMemory, SLATECORE_API);
-
-
 // Compile all the RichText and MultiLine editable text?
 #define WITH_FANCY_TEXT 1
 
diff --git a/Engine/Source/Runtime/SlateCore/Public/Stats/SlateStats.h b/Engine/Source/Runtime/SlateCore/Public/Stats/SlateStats.h
new file mode 100644
--- /dev/null
+++ b/Engine/Source/Runtime/SlateCore/Public/Stats/SlateStats.h
@@ -0,0 +1,532 @@
+// Copyright 1998-2014 Epic Games, Inc. All Rights Reserved.
+
+#pragma once
+
+/** Set to zero to disable the slate stats system completely, meaning all stats-related structures and code are compiled out. */
+#define SLATE_STATS !UE_BUILD_SHIPPING || WITH_EDITOR
+
+/** 
+ * Predefined detail levels to make choosing level slightly easier.
+ * The _OFF setting effectively turns off any stat while still keeping the stats system compiled in. 
+ */
+#define SLATE_STATS_DETAIL_LEVEL_FULL INT_MAX
+#define SLATE_STATS_DETAIL_LEVEL_HI INT_MAX/4*3
+#define SLATE_STATS_DETAIL_LEVEL_MED INT_MAX/4*2
+#define SLATE_STATS_DETAIL_LEVEL_LOW INT_MAX/4*1
+#define SLATE_STATS_DETAIL_LEVEL_OFF INT_MIN+1
+
+/** Controls what level of detail of stats should be compiled in. Detail level is a simple int. Any stat with a detail level less than this value will be compiled in. */
+#define SLATE_STATS_DETAIL_LEVEL SLATE_STATS_DETAIL_LEVEL_MED
+
+#if SLATE_STATS
+
+/** 
+ * Slate Stats System Motivation
+ * ---------------------------
+ * Slate uses it's own stats system to gather performance metrics. There are a few reasons it does not use the builtin UE4 stats system:
+ *   1. UE4 stats are threadsafe, thus heavier weight. Slate is single-threaded by nature due to UI constraints.
+ *   2. Slate stats can optionally carry additional context (like class or debug info) used to provide more info when digesting the profile.
+ *   3. Using a lighter weight system allows us to add more tracking by default.
+ *   4. Slate uses a custom hierarchical system that allows it to spit out a CSV containing a complete hierarchical inclusive/exclusive profile.
+ * 
+ * Slate Stats System Overview
+ * ---------------------------
+ * You need to declare your stat globally somewhere. It should generally be a static variable:
+ * 
+ *   SLATE_DECLARE_CYCLE_COUNTER(GSlateOnPaint, "OnPaint");
+ * 
+ * Next you go the scope of code to provide and add a scope counter (the types will be explained below).
+ * 
+ *   SLATE_CYCLE_COUNTER_SCOPE(GSlateOnPaint);
+ * 
+ * This enables the Slate Stats system to track this scope of code and associate it with the "OnPaint" stat.
+ * If you want to track the OnPaint call of each type of widget separately, you would use the _CUSTOM version of the macro and provide it with
+ * an FName to track along with the stat like so (assuming "this" is an SWidget):
+ * 
+ *   SLATE_CYCLE_COUNTER_SCOPE_CUSTOM(GSlateOnPaint, GetType());
+ * 
+ * However, note that custom info is only available when doing a full hierarchical capture (as explained below).
+ * 
+ * Slate supports two modes of stats - flat and hierarchical. Stats are hierarchical by default. Use the _FLAT form of the macro 
+ * to declare a stat that is only used for flattened tracking. This is faster when you don't need this stat for a full hierarchical profile.
+ *
+ *   Flat
+ *   ----
+ *   Flat view is a standard inclusive sum of cycles. Each scope simply records a start and end time and adds it to the global stat accumulator.
+ *   This happens until the system is told to display the values, at which point it averages the cycle accumulators over the number of frames since
+ *   the last summary. The resultant view is essentially a list of all declared stats along with the average inclusive time spent in that stat's scope. 
+ *   
+ *   The system handles re-entrant stats. So if you are processing a hierarchy (as Slate widgets essentially are), re-entering the same stat will do
+ *   nothing until the last scope is exited.
+ *   
+ *   Hierarchical
+ *   ------------
+ *   Hierarchical capture essentially creates a complete hierarchical profile of a frame. Since it has to track EACH stat entry that frame, it can
+ *   have a non-trivial performance and memory impact on the application. However, the result is similar to the data from a fully instrumented profile
+ *   of your application, with a few notable details:
+ *     1. You are instrumenting sections of code, not functions. This allows for finer-grained classification of your code.
+ *     2. each entry also carries any context you passed it at runtime, allowing for even better context than an instrumented profiler can give you.
+ *     3. Stats are not subject to the whims of inlining, etc, since they are injected at compile-time.
+ *   While Slate goes out of its way to make this capture as cheap as possible, there is a lot of data being gathered each frame. Therefore, Slate does NOT
+ *   do a hierarchical capture continually. Instead, the stats system must be triggered (programmatically) to enable hierarchical capture right before a 
+ *   frame is started. This is done when calling EndFrame() on the FSlateStatHierarchy singleton:
+ *   
+ *     FSlateStatHierarchy::Get().EndFrame(true|false);
+ *     
+ *   This can be triggered by end users through the Slate.Stats.Hierarchy.Trigger cvar.
+ *   
+ *   Since there is no guarantee that the stat call stack will be equivalent from one frame to the next, we can only accumulate ONE frame.
+ *   
+ *   When a hierarchical capture frame ends, Slate does some additional work to gather up all instances, derive "full paths" to each entry, and derive
+ *   the exclusive time for that instance. Again, Slate tries to do this as fast as possible, but there can be 100,000s of entries in a frame, so that
+ *   can only be as fast as those entries can be traversed. When finished, Slate spits out a CSV into the Saved/ folder containing the following fields:
+ *     1. Path - full path to the entry. 
+ *     2. Stat Description
+ *     3. Stat additional context (Custom info)
+ *     4. Inclusive time
+ *     5. Exclusive time
+ *     
+ *   Path is essentially a pathway through the hierarchy of samples to the entry (a stat call-stack if you will). This allows analysis tools to understand 
+ *   the capture hierarchy so they can do things like restrict to subtrees of the capture, create caller-callee views, etc. Essentially typical instrumented 
+ *   profiler things.
+ * 
+ * FSlateStatCycleCounter - Global storage of a cycle stat. This should be declared as a global using the SLATE_DECLARE_CYCLE_COUNTER macro.
+ * Counters don't update themselves, they rely on the scope classes to update their values.
+ * All the methods of this class are essentially internal and there is no need for an end-user to call them directly.
+ */
+class FSlateStatCycleCounter
+{
+public:
+	/** 
+	 * Default ctor. Registers itself with the global list of counters (so the system can iterate over them).
+	 * @param InCounterName Name of the counter. Essentially an arbitrary description string, but try to keep it terse.
+	 */
+	SLATECORE_API FSlateStatCycleCounter(FName InCounterName);
+	
+	/**
+	 * Unregisteres the counter from the global list. Uses a global "memory pool" to avoid initialization/shutdown dependencies.
+	 */
+	SLATECORE_API ~FSlateStatCycleCounter();
+
+	/** Public access to report on a stat. Time is pre-converted to milliseconds. */
+	double GetLastComputedAverageInclusiveTime() const
+	{
+		return LastComputedAverageInclusiveTime;
+	}
+
+	/** Public access to report on a stat. */
+	FName GetName() const
+	{
+		return Name;
+	}
+
+	/** Access to the global list of counters to iterate over. Internally uses a memory pool (not a static instance) to avoid shutdown early-destruction issues. */
+	static SLATECORE_API const TArray<FSlateStatCycleCounter*>& GetRegisteredCounters();
+
+	/** Perform all end-frame work. Slate makes this call automatically for you, so end users should not have to call it. */
+	static SLATECORE_API void EndFrame(double CurrentTime);
+
+	static SLATECORE_API bool AverageInclusiveTimesWereUpdatedThisFrame();
+
+private:
+	/** Needs to be called after the stats have been reported. This resets the accumulated time. */
+	SLATECORE_API void Reset();
+
+	/** 
+	 * A global list of counters that all instances register with. Allows easier iteration. 
+	 * Internally creates a dynamic allocation and "leaks" it (aka: memory pool) so it doesn't get destroyed in the wrong order at shutdown.
+	 * Declared out of line so we can have a single instance across module boundaries.
+	 */
+	static SLATECORE_API TArray<FSlateStatCycleCounter*>& GetGlobalRegisteredCounters();
+
+	// let these classes mutate the counter values.
+	template<bool> friend class FSlateStatCycleCounterScopeFlat;
+	template<bool> friend class FSlateStatCycleCounterScopeHierarchical;
+
+	/** Display name of the stat. */
+	FName Name;
+	/** Accumulated time since the last report. Something needs to track frames since last report for this to be useful. */
+	double InclusiveTime;
+	/** Last averaged stat computed by the stats system. This is the value that can be used for display purposes. */
+	double LastComputedAverageInclusiveTime;
+
+	/** 
+	 * Used by the scope classes. 
+	 * This really makes the system single-threaded only because we can only track one instance of a counter at a time like this.
+	 * Saves us from having to look it up.
+	 * 
+	 * When a timer is re-entrant, we don't want to keep starting the time over. So we only track the outer level. 
+	 */
+	int32 StackDepth;
+	/** 
+	 * Used by the scope classes. 
+	 * This really makes the system single-threaded only because we can only track one instance of a counter at a time like this.
+	 * Saves us from having to look it up.
+	 * 
+	 * Tracks when we started timing this counter so we can compute a delta.
+	 */
+	double StartTime;
+};
+
+/** 
+ * Scoped tracking of a FSlateStatCycleCounter. 
+ * Super-lean, starts timing in the ctor, and stops in the dtor. 
+ * Only tracks time when it is at the top level, effectively ignoring re-entrancy. 
+ * 
+ * Inlined purposefully to keep the call overhead down.
+ */
+template <bool bShouldBeCompiledIn>
+class FSlateStatCycleCounterScopeFlat
+{
+public:
+	/** Ctor connects this scope to a specific counter to accumulate to and starts the timing. */
+	FSlateStatCycleCounterScopeFlat( FSlateStatCycleCounter& InCounter )
+		: Counter(InCounter)
+	{
+		extern SLATECORE_API int32 GSlateStatsFlatEnable;
+		if (GSlateStatsFlatEnable > 0)
+		{
+			StartTiming();
+		}
+	}
+	/** Dtor stops the timing and accumulates. */
+	~FSlateStatCycleCounterScopeFlat()
+	{
+		extern SLATECORE_API int32 GSlateStatsFlatEnable;
+		if (GSlateStatsFlatEnable > 0)
+		{
+			StopTiming();
+		}
+	}
+
+private:
+	/** Starts the timing. Ignores the call if we detect re-entrancy. */
+	void StartTiming()
+	{
+		// only start timing when the stack depth is zero.
+		if (Counter.StackDepth++ == 0)
+		{
+			Counter.StartTime = FPlatformTime::Seconds();
+		}
+	}
+
+	/** Stops the timing. Ignores the call if we detect re-entrancy. */
+	void StopTiming()
+	{
+		// only stop timing when the stack depth goes back to zero, meaning we're finally out of our re-entrant stack.
+		if (--Counter.StackDepth == 0)
+		{
+			const double Delta = FPlatformTime::Seconds() - Counter.StartTime;
+			Counter.InclusiveTime += Delta;
+		}
+	}
+
+	// stores a reference to the stat we will be accumulating to.
+	FSlateStatCycleCounter& Counter;
+};
+
+/** Class only exists so stats that don't meet the defined SLATE_STATS_DETAIL_LEVEL are not compiled in. */
+template <>
+class FSlateStatCycleCounterScopeFlat<false>
+{
+public:
+	FSlateStatCycleCounterScopeFlat(FSlateStatCycleCounter&) {}
+};
+
+
+
+/** Tracks info for a hierarchical stats instance (explained in FSlateStatHierarchy). */
+struct FSlateStatHierarchyEntry
+{
+	/** Initializes all the data needed for a hierarchical entry. */
+	FSlateStatHierarchyEntry(FName InCounterName, FName InCustomName, double StartTime, int32 InStackDepth) : CounterName(InCounterName), CustomName(InCustomName), InclusiveTime(StartTime), ExclusiveTime(0.0), StackDepth(InStackDepth) {}
+	/** Counter display name. Same as flattened stats. */
+	FName CounterName;
+	/** Custom context string. Used for things like printing RTTI when tracking a virtual function hierarchically. */
+	FName CustomName;
+	/** Time measured for this stat. */
+	double InclusiveTime;
+	/** Exclusive time measured for this stat. Not computed immediately. Has to be calculated when the frame is over. */
+	double ExclusiveTime;
+	/** Depth of this stat in the stat hierarchy. Will be used during reporting to unwind the stack and compute a full path for each entry in the list. */
+	int32 StackDepth;
+};
+
+/** 
+ * Hierarchical tracking of stats.
+ * Heavier weight than flattened tracking done by FSlateStatCycleCounter because it has to track EACH scope separately in a dynamic array.
+ * Due to this, we don't track this every frame, and only show the results of a single frame capture that we decide in advance to track
+ * (ie, we can't do something like hitch detection, which is only disovered in post).
+ * 
+ * Still fairly lean as far as hierarchical tracking goes, but does a lot more timing calls in highly re-entrant stats
+ * (where the flattened system only tracks the top level).
+ * 
+ * Like flattened stats, all of this only works correctly when stats are gathered on a single thread (the same thread that reports them)!
+ */
+class FSlateStatHierarchy
+{
+public:
+	/** Class is a singleton, so we reserve some space for the hierarchy when it is first constructed. */
+	FSlateStatHierarchy();
+
+	/** 
+	 * Singleton access. 
+	 * Declared out of line so we can have a single instance across module boundaries.
+	 */
+	static SLATECORE_API FSlateStatHierarchy& Get();
+
+	/** 
+	 * Recursive function to traverse the entries at the end of a frame and compute the exclusive times as quickly as possible.
+	 * The basic premise is to increment of each entry (Index param), while tracking the index of the immediate parent of that 
+	 * entry.
+	 *   * If we are one level deeper than our parent, then all the inclusive time used by this entry is NOT part of the 
+	 * exclusive time of the parent, so we subtract it off.
+	 *   * When we descend more than one level deeper than the parent, we know that the preceding entry is the parent, as all
+	 *     entries are added in order (and single-threaded).
+	 *   * As we ascend from our recursive stack we update the index so the calling code knows to skip over the descendents we
+	 *     already processed.
+	 *     
+	 * This results in an extremely quick processing of the entries, mostly bound by our stack depth.
+	 * 
+	 * @param Index - index of the entry we are processing 
+	 * @param ParentIndex - index of this entry's parent. May have ot be updated if we discover this entry has descended a level in the stack.
+	 * @return the Index AFTER the last one we processed.
+	 */
+	int32 ComputeExclusiveTimes(int32 Index = 1, int32 ParentIndex = 0)
+	{
+		const int32 NumEntries = StatEntries.Num();
+		// keep going until we reach the end.
+		while (Index < NumEntries)
+		{
+			// grab the info about our parent.
+			FSlateStatHierarchyEntry& ParentEntry = StatEntries[ParentIndex];
+			const int32 ParentDepth = ParentEntry.StackDepth;
+			// look at our current depth.
+			const int32 Depth = StatEntries[Index].StackDepth;
+			// we are immediately below our parent. Subtract off the entrie's inclusive time
+			// because it isn't part of the parent's exclusive time.
+			if (Depth == ParentDepth + 1)
+			{
+				const double InclusiveTime = StatEntries[Index].InclusiveTime;
+				ParentEntry.ExclusiveTime -= InclusiveTime;
+				++Index;
+			}
+			// we have descended on level in the stack. Update the parent and go again.
+			else if (Depth == ParentDepth + 2)
+			{
+				Index = ComputeExclusiveTimes(Index, Index - 1);
+			}
+			// we are ascending out of the stack, so just return.
+			else
+			{
+				return Index;
+			}
+		
+		}
+		// we've reached the end of our list of entries.
+		return Index;
+	}
+
+	/** 
+	 * Needs to be called every time a frame ends. 
+	 * Empties the stack of stats (without reclaiming the memory). 
+	 * Set bTrackNextFrame to true to actually enable hierarchical tracking on the next frame. 
+	 */
+	void EndFrame(bool bTrackNextFrame);
+
+	/** tracks whether we are actually gathering hierarchical stats this frame. */
+	bool IsTrackingThisFrame() const
+	{
+		return bTrackThisFrame;
+	}
+
+	/** Access to the hierarhcy for reporting. */
+	const TArray<FSlateStatHierarchyEntry>& GetStatEntries() const { return StatEntries; }
+
+	/** Interface to start tracking an instance. Returns the tracking instance so we don't have to look it up later.. */
+	FSlateStatHierarchyEntry* StartStat(FName CounterName, FName CustomName, double StartTime)
+	{
+		if (!bTrackThisFrame) return NULL;
+		checkfSlow(StatEntries.Num() >= StatEntries.Max(), TEXT("You have overrun the Slate Stats hierarhical profiling entries limit. Profile a simpler scene or adjust the #define SLATE_STATS_HIERARCHY_MAX_ENTRIES to a larger value and recompile your code."));
+		StatEntries.Emplace(CounterName, CustomName, StartTime, StackDepth++);
+		return &StatEntries.Last();
+	}
+
+	/** Called when done tracking an instance. Pass the entry returned by StartStat(). */
+	void StopStat(FSlateStatHierarchyEntry* Entry, double EndTime)
+	{
+		if (!Entry) return;
+		Entry->InclusiveTime = EndTime - Entry->InclusiveTime;
+		// This isn't the real exclusive time yet. At the end of the frame we will subtract off child times.
+		Entry->ExclusiveTime = Entry->InclusiveTime;
+		// indicate that we are ascending our stack of entries.
+		--StackDepth;
+	}
+
+private:
+	/** 
+	 * Stat Entries measured this frame. Usually empty unless EndFrame() is called with bTrackNextFrame = true. 
+	 * 
+	 * WARNING: This array SHOULD NOT be resized in a frame. If it is, the system will assert or crash in non-checkslow builds.
+	 * If this happens, adjust SLATE_STATS_HIERARCHY_MAX_ENTRIES for your needs.
+	 */
+	TArray<FSlateStatHierarchyEntry> StatEntries;
+	/** Current stat hierarchy depth. */
+	int32 StackDepth;
+	/** If false, does not actually track the hierarchy that frame. */
+	bool bTrackThisFrame;
+
+};
+
+/** 
+ * Tracks a flattened stat, but also adds to the hierarchical tracker. 
+ * Supports a custom FName to be passed in to track some context specific info about this stat in the hierarchy (like a specific class name when tracking a virtual method).
+ * This is a completely separate class to keep the flat version super lean. 
+ * 
+ * Inlined purposefully to keep the call overhead down.
+ */
+template <bool bShouldBeCompiledIn>
+class FSlateStatCycleCounterScopeHierarchical
+{
+public:
+	/** Ctor mimics flat stats logic, but also allocates a hierarchical entry if we are capturing hierarchical stats this frame. */
+	FSlateStatCycleCounterScopeHierarchical( FSlateStatCycleCounter& InCounter )
+		: Counter(InCounter)
+		, HierarchyEntry(nullptr)
+	{
+		StartTiming(NAME_None);
+	}
+	/**
+	 * Ctor mimics flat stats logic, but also allocates a hierarchical entry if we are capturing hierarchical stats this frame. 
+	 * This version allows passing a custom FName to associate with the entry.
+	 */
+	FSlateStatCycleCounterScopeHierarchical( FSlateStatCycleCounter& InCounter, FName CustomName )
+		: Counter(InCounter)
+		, HierarchyEntry(nullptr)
+	{
+		StartTiming(CustomName);
+	}
+	/** Dtor mimics flat stats logic, but also updates the hierarchical entry if we are capturing hierarchical stats this frame. */
+	~FSlateStatCycleCounterScopeHierarchical()
+	{
+		StopTiming();
+	}
+
+private:
+	/**
+	 * Starts the hierarchical timing. If we are not tracking this frame, falls back to simple flat tracking.
+	 * The main difference is that re-entrant stats entries ARE tracked individually by the hierarchical tracker.
+	 */
+	void StartTiming(FName CustomName)
+	{
+		// Skip a bunch of work if we are not tracking the hierarchy this frame, essentially falling back to the flat stat logic.
+		if (FSlateStatHierarchy::Get().IsTrackingThisFrame())
+		{
+			const double CurrentTime = FPlatformTime::Seconds();
+			if (Counter.StackDepth++ == 0)
+			{
+				Counter.StartTime = CurrentTime;
+			}
+			HierarchyEntry = FSlateStatHierarchy::Get().StartStat(Counter.Name, CustomName, CurrentTime);
+		}
+		else
+		{
+			extern SLATECORE_API int32 GSlateStatsFlatEnable;
+			if (GSlateStatsFlatEnable > 0)
+			{
+				if (Counter.StackDepth++ == 0)
+				{
+					Counter.StartTime = FPlatformTime::Seconds();
+				}
+			}
+		}
+	}
+
+	/**
+	 * Stops the hierarchical timing. If we are not tacking this frame, falls back to simple flat tracking.
+	 */
+	void StopTiming()
+	{
+		// Skip a bunch of work if we are not tracking the hierarchy this frame, essentially falling back to the flat stat logic.
+		if (HierarchyEntry != NULL)
+		{
+			const double CurrentTime = FPlatformTime::Seconds();
+			FSlateStatHierarchy::Get().StopStat(HierarchyEntry, CurrentTime);
+			// only increment the stat when the stack depth goes back to zero, meaning we're finally out of our re-entrant stack.
+			if (--Counter.StackDepth == 0)
+			{
+				const double Delta = CurrentTime - Counter.StartTime;
+				Counter.InclusiveTime += Delta;
+			}
+		}
+		else
+		{
+			extern SLATECORE_API int32 GSlateStatsFlatEnable;
+			if (GSlateStatsFlatEnable > 0)
+			{
+				// only increment the stat when the stack depth goes back to zero, meaning we're finally out of our re-entrant stack.
+				if (--Counter.StackDepth == 0)
+				{
+					const double Delta = FPlatformTime::Seconds() - Counter.StartTime;
+					Counter.InclusiveTime += Delta;
+				}
+			}
+		}
+	}
+
+	/** The counter we are tracking for. */
+	FSlateStatCycleCounter& Counter;
+	/** The actual hierarchical entry we are recording to. This stores all the real data. */
+	FSlateStatHierarchyEntry* HierarchyEntry;
+};
+
+/** Class only exists so stats that don't meet the defined SLATE_STATS_DETAIL_LEVEL are not compiled in. */
+template <>
+class FSlateStatCycleCounterScopeHierarchical < false >
+{
+public:
+	FSlateStatCycleCounterScopeHierarchical(FSlateStatCycleCounter&) {}
+	FSlateStatCycleCounterScopeHierarchical(FSlateStatCycleCounter&, FName) {}
+};
+
+/** Declares a Slate Stat cycle counter. */
+#define SLATE_DECLARE_CYCLE_COUNTER(CounterName, CounterDescription) \
+	FSlateStatCycleCounter CounterName(CounterDescription)
+
+/** Track a Slate Stat, also adding it to the hierarchy if we are capturing hierarchical stats this frame. */
+#define SLATE_CYCLE_COUNTER_SCOPE(StatCounterWes) \
+	FSlateStatCycleCounterScopeHierarchical<(SLATE_STATS_DETAIL_LEVEL != SLATE_STATS_DETAIL_LEVEL_OFF)> TrackStat_##StatCounterWes(StatCounterWes)
+
+/** Track a Slate Stat, also adding it to the hierarchy if we are capturing hierarchical stats this frame. With a given detail level so it can be compiled out. */
+#define SLATE_CYCLE_COUNTER_SCOPE_DETAILED(DetailLevel, StatCounterWes) \
+	FSlateStatCycleCounterScopeHierarchical<(DetailLevel <= SLATE_STATS_DETAIL_LEVEL) && (DetailLevel != SLATE_STATS_DETAIL_LEVEL_OFF)> TrackStat_##StatCounterWes(StatCounterWes)
+
+/** Track a Slate Stat, also adding it to the hierarchy if we are capturing hierarchical stats this frame, along with custom context. */
+#define SLATE_CYCLE_COUNTER_SCOPE_CUSTOM(StatCounterWes,CustomName) \
+	FSlateStatCycleCounterScopeHierarchical<(SLATE_STATS_DETAIL_LEVEL != SLATE_STATS_DETAIL_LEVEL_OFF)> TrackStat_##StatCounterWes(StatCounterWes, CustomName)
+
+/** Track a Slate Stat, also adding it to the hierarchy if we are capturing hierarchical stats this frame, along with custom context. With a given detail level so it can be compiled out. */
+#define SLATE_CYCLE_COUNTER_SCOPE_CUSTOM_DETAILED(DetailLevel, StatCounterWes,CustomName) \
+	FSlateStatCycleCounterScopeHierarchical<(DetailLevel <= SLATE_STATS_DETAIL_LEVEL) && (DetailLevel != SLATE_STATS_DETAIL_LEVEL_OFF)> TrackStat_##StatCounterWes(StatCounterWes, CustomName)
+
+/** Track a Slate Stat, only for flattened display. Slightly cheaper than checking dynamically. It doesn't make sense to have a custom version as flattened stats don't gather custom data right now. */
+#define SLATE_CYCLE_COUNTER_SCOPE_FLAT(StatCounterWes) \
+	FSlateStatCycleCounterScopeFlat<(SLATE_STATS_DETAIL_LEVEL != SLATE_STATS_DETAIL_LEVEL_OFF)> TrackStat_##StatCounterWes(StatCounterWes)
+
+/** Track a Slate Stat, only for flattened display. Slightly cheaper than checking dynamically. With a given detail level so it can be compiled out. */
+#define SLATE_CYCLE_COUNTER_SCOPE_FLAT_DETAILED(DetailLevel, StatCounterWes) \
+	FSlateStatCycleCounterScopeFlat<(DetailLevel <= SLATE_STATS_DETAIL_LEVEL) && (DetailLevel != SLATE_STATS_DETAIL_LEVEL_OFF)> TrackStat_##StatCounterWes(StatCounterWes)
+
+#define SLATE_STATS_END_FRAME(CurrentTime) FSlateStatCycleCounter::EndFrame(CurrentTime)
+
+#else
+#define SLATE_DECLARE_CYCLE_COUNTER(CounterName, CounterDescription)
+#define SLATE_CYCLE_COUNTER_SCOPE(StatCounterWes)
+#define SLATE_CYCLE_COUNTER_SCOPE_DETAILED(DetailLevel,StatCounterWes)
+#define SLATE_CYCLE_COUNTER_SCOPE_CUSTOM(StatCounterWes,CustomName)
+#define SLATE_CYCLE_COUNTER_SCOPE_CUSTOM_DETAILED(DetailLevel,StatCounterWes,CustomName)
+#define SLATE_CYCLE_COUNTER_SCOPE_FLAT(StatCounterWes)
+#define SLATE_CYCLE_COUNTER_SCOPE_FLAT_DETAILED(DetailLevel,StatCounterWes)
+#define SLATE_STATS_END_FRAME(CurrentTime)
+
+#endif
diff --git a/Engine/Source/Runtime/SlateCore/Public/Textures/SlateShaderResource.h b/Engine/Source/Runtime/SlateCore/Public/Textures/SlateShaderResource.h
--- a/Engine/Source/Runtime/SlateCore/Public/Textures/SlateShaderResource.h
+++ b/Engine/Source/Runtime/SlateCore/Public/Textures/SlateShaderResource.h
@@ -2,6 +2,7 @@
 
 #pragma once
 
+DECLARE_MEMORY_STAT_EXTERN(TEXT("Texture Data GPU Memory"), STAT_SlateTextureGPUMemory, STATGROUP_SlateMemory, SLATECORE_API);
 
 namespace ESlateShaderResource
 {
diff --git a/Engine/Source/Runtime/SlateCore/Public/Textures/SlateTextureData.h b/Engine/Source/Runtime/SlateCore/Public/Textures/SlateTextureData.h
--- a/Engine/Source/Runtime/SlateCore/Public/Textures/SlateTextureData.h
+++ b/Engine/Source/Runtime/SlateCore/Public/Textures/SlateTextureData.h
@@ -2,6 +2,7 @@
 
 #pragma once
 
+DECLARE_MEMORY_STAT_EXTERN(TEXT("Texture Data CPU Memory"), STAT_SlateTextureDataMemory, STATGROUP_SlateMemory, SLATECORE_API);
 
 /**
  * Holds texture data for upload to a rendering resource
diff --git a/Engine/Source/Runtime/SlateCore/Public/Widgets/SWidget.h b/Engine/Source/Runtime/SlateCore/Public/Widgets/SWidget.h
--- a/Engine/Source/Runtime/SlateCore/Public/Widgets/SWidget.h
+++ b/Engine/Source/Runtime/SlateCore/Public/Widgets/SWidget.h
@@ -35,6 +35,8 @@ public:
 	}
 };
 
+class IToolTip;
+
 /**
  * Abstract base class for Slate widgets.
  *
@@ -643,10 +645,7 @@ public:
 	}
 
 	/** @return is this widget visible, hidden or collapsed */
-	EVisibility GetVisibility() const
-	{
-		return Visibility.Get();
-	}
+	EVisibility GetVisibility() const;
 
 	/** @param InVisibility  should this widget be */
 	virtual void SetVisibility( TAttribute<EVisibility> InVisibility )
@@ -770,6 +769,9 @@ public:
 	/** @return A String of the widget's code location */
 	virtual FString GetCreatedInFile() const;
 
+	/** @return A String of the widget's code location */
+	virtual FName GetCreatedInFileFName() const;
+
 	/** @return The line number of the widgets location */
 	virtual int32 GetCreatedInLineNumber() const;
 
diff --git a/Engine/Source/Runtime/SlateRHIRenderer/Private/SlateRHIRenderer.cpp b/Engine/Source/Runtime/SlateRHIRenderer/Private/SlateRHIRenderer.cpp
--- a/Engine/Source/Runtime/SlateRHIRenderer/Private/SlateRHIRenderer.cpp
+++ b/Engine/Source/Runtime/SlateRHIRenderer/Private/SlateRHIRenderer.cpp
@@ -11,6 +11,8 @@ DECLARE_CYCLE_STAT(TEXT("Map Staging Buffer"),STAT_MapStagingBuffer,STATGROUP_Cr
 DECLARE_CYCLE_STAT(TEXT("Generate Capture Buffer"),STAT_GenerateCaptureBuffer,STATGROUP_CrashTracker);
 DECLARE_CYCLE_STAT(TEXT("Unmap Staging Buffer"),STAT_UnmapStagingBuffer,STATGROUP_CrashTracker);
 
+DECLARE_CYCLE_STAT(TEXT("Slate Rendering RT Time"), STAT_SlateRenderingRTTime, STATGROUP_Slate);
+DECLARE_CYCLE_STAT(TEXT("Slate RT Present Time"), STAT_SlatePresentRTTime, STATGROUP_Slate);
 
 // Defines the maximum size that a slate viewport will create
 #define MAX_VIEWPORT_SIZE 16384
@@ -554,8 +556,6 @@ void FSlateRHIRenderer::PrepareToTakeScreenshot(const FIntRect& Rect, TArray<FCo
  */
 void FSlateRHIRenderer::DrawWindows_Private( FSlateDrawBuffer& WindowDrawBuffer )
 {
-	SCOPE_CYCLE_COUNTER( STAT_SlateRenderingGTTime );
-
 	check( IsThreadSafeForSlateRendering() );
 
 	// Enqueue a command to unlock the draw buffer after all windows have been drawn
diff --git a/Engine/Source/Runtime/SlateRHIRenderer/Private/SlateRHIRenderingPolicy.cpp b/Engine/Source/Runtime/SlateRHIRenderer/Private/SlateRHIRenderingPolicy.cpp
--- a/Engine/Source/Runtime/SlateRHIRenderer/Private/SlateRHIRenderingPolicy.cpp
+++ b/Engine/Source/Runtime/SlateRHIRenderer/Private/SlateRHIRenderingPolicy.cpp
@@ -11,6 +11,12 @@
 #include "SlateUTextureResource.h"
 #include "SlateMaterialResource.h"
 
+DECLARE_CYCLE_STAT(TEXT("Update Buffers RT"), STAT_SlateUpdateBufferRTTime, STATGROUP_Slate);
+DECLARE_CYCLE_STAT(TEXT("Draw Time"), STAT_SlateDrawTime, STATGROUP_Slate);
+
+DECLARE_MEMORY_STAT(TEXT("Vertex Buffer Memory"), STAT_SlateVertexBufferMemory, STATGROUP_SlateMemory);
+DECLARE_MEMORY_STAT(TEXT("Index Buffer Memory"), STAT_SlateIndexBufferMemory, STATGROUP_SlateMemory);
+
 FSlateElementIndexBuffer::FSlateElementIndexBuffer()
 	: BufferSize(0)	 
 	, BufferUsageSize(0)
diff --git a/Engine/Source/Runtime/SlateRHIRenderer/Private/SlateRHIResourceManager.cpp b/Engine/Source/Runtime/SlateRHIRenderer/Private/SlateRHIResourceManager.cpp
--- a/Engine/Source/Runtime/SlateRHIRenderer/Private/SlateRHIResourceManager.cpp
+++ b/Engine/Source/Runtime/SlateRHIRenderer/Private/SlateRHIResourceManager.cpp
@@ -6,6 +6,10 @@
 #include "SlateUTextureResource.h"
 #include "SlateMaterialResource.h"
 
+DECLARE_DWORD_ACCUMULATOR_STAT(TEXT("Num Texture Atlases"), STAT_SlateNumTextureAtlases, STATGROUP_SlateMemory);
+DECLARE_DWORD_ACCUMULATOR_STAT(TEXT("Num Non-Atlased Textures"), STAT_SlateNumNonAtlasedTextures, STATGROUP_SlateMemory);
+DECLARE_DWORD_ACCUMULATOR_STAT(TEXT("Num Dynamic Textures"), STAT_SlateNumDynamicTextures, STATGROUP_SlateMemory);
+
 TSharedPtr<FSlateDynamicTextureResource> FDynamicResourceMap::GetDynamicTextureResource( FName ResourceName ) const
 {
 	return NativeTextureMap.FindRef( ResourceName );
@@ -178,8 +182,6 @@ bool FSlateRHIResourceManager::IsAtlasPageResourceAlphaOnly() const
 
 void FSlateRHIResourceManager::CreateTextures( const TArray< const FSlateBrush* >& Resources )
 {
-	DECLARE_SCOPE_CYCLE_COUNTER(TEXT("Loading Slate Textures"), STAT_Slate, STATGROUP_LoadTime);
-
 	TMap<FName,FNewTextureInfo> TextureInfoMap;
 
 	const uint32 Stride = GPixelFormats[PF_R8G8B8A8].BlockBytes;
@@ -759,7 +761,7 @@ void FSlateRHIResourceManager::DeleteResources()
 	{
 		delete NonAtlasedTextures[ResourceIndex];
 	}
-	SET_DWORD_STAT(STAT_SlateNumNonAtlasedTextures, 0 );
+	SET_DWORD_STAT(STAT_SlateNumNonAtlasedTextures, 0);
 	SET_DWORD_STAT(STAT_SlateNumTextureAtlases, 0);
 	SET_DWORD_STAT(STAT_SlateNumDynamicTextures, 0);
 
