Fix a block of code not surrounded by SLATE_STATS

diff --git a/Engine/Source/Developer/SlateReflector/Private/Widgets/SWidgetReflector.cpp b/Engine/Source/Developer/SlateReflector/Private/Widgets/SWidgetReflector.cpp
--- a/Engine/Source/Developer/SlateReflector/Private/Widgets/SWidgetReflector.cpp
+++ b/Engine/Source/Developer/SlateReflector/Private/Widgets/SWidgetReflector.cpp
@@ -635,11 +635,13 @@ void SWidgetReflector::Construct( const FArguments& InArgs )
 			]
 	];
 
+#if SLATE_STATS			
 	if (StatsToolsBox.IsValid())
 	{
 		extern SLATECORE_API int32 GSlateStatsFlatEnable;
 		StatsToolsBox->SetVisibility(GSlateStatsFlatEnable != 0 ? EVisibility::Visible : EVisibility::Collapsed);
 	}
+#endif
 }
 
 /* SCompoundWidget overrides
