CIS fix for Mac and Slate Stats
#mac

diff --git a/Engine/Source/Developer/SlateReflector/Private/Widgets/SWidgetReflector.cpp b/Engine/Source/Developer/SlateReflector/Private/Widgets/SWidgetReflector.cpp
--- a/Engine/Source/Developer/SlateReflector/Private/Widgets/SWidgetReflector.cpp
+++ b/Engine/Source/Developer/SlateReflector/Private/Widgets/SWidgetReflector.cpp
@@ -8,6 +8,12 @@
 #define WITH_EVENT_LOGGING 0
 
 extern SLATECORE_API int32 bFoldTick;
+#if SLATE_STATS
+extern SLATECORE_API int32 GSlateStatsFlatEnable;
+extern SLATECORE_API int32 GSlateStatsFlatLogOutput;
+extern SLATECORE_API int32 GSlateStatsHierarchyTrigger;
+extern SLATECORE_API float GSlateStatsFlatIntervalWindowSec;
+#endif
 static const int32 MaxLoggedEvents = 100;
 
 /**
@@ -407,12 +413,10 @@ void SWidgetReflector::Construct( const FArguments& InArgs )
 						.Style( FCoreStyle::Get(), "ToggleButtonCheckbox" )
 						.IsChecked_Static([]
 						{
-							extern SLATECORE_API int32 GSlateStatsFlatEnable;
 							return GSlateStatsFlatEnable == 0 ? ECheckBoxState::Unchecked : ECheckBoxState::Checked;
 						})
 						.OnCheckStateChanged_Lambda( [=]( const ECheckBoxState NewState )
 						{
-							extern SLATECORE_API int32 GSlateStatsFlatEnable;
 							GSlateStatsFlatEnable = (NewState == ECheckBoxState::Checked) ? 1 : 0;
 							StatsToolsBox->SetVisibility(GSlateStatsFlatEnable != 0 ? EVisibility::Visible : EVisibility::Collapsed);
 						})
@@ -478,12 +482,10 @@ void SWidgetReflector::Construct( const FArguments& InArgs )
 						.Style( FCoreStyle::Get(), "ToggleButtonCheckbox" )
 						.IsChecked_Static([]
 						{
-							extern SLATECORE_API int32 GSlateStatsFlatLogOutput;
 							return GSlateStatsFlatLogOutput == 0 ? ECheckBoxState::Unchecked : ECheckBoxState::Checked;
 						})
 						.OnCheckStateChanged_Static([]( const ECheckBoxState NewState )
 						{
-							extern SLATECORE_API int32 GSlateStatsFlatLogOutput;
 							GSlateStatsFlatLogOutput = (NewState == ECheckBoxState::Checked) ? 1 : 0;
 						})
 						.ToolTip
@@ -512,7 +514,6 @@ void SWidgetReflector::Construct( const FArguments& InArgs )
 						SNew(SButton)
 						.OnClicked_Static([]
 						{
-							extern SLATECORE_API int32 GSlateStatsHierarchyTrigger;
 							GSlateStatsHierarchyTrigger = 1;
 							return FReply::Handled();
 						})
@@ -560,7 +561,6 @@ void SWidgetReflector::Construct( const FArguments& InArgs )
 						)
 						.Value_Static([]
 						{
-							extern SLATECORE_API float GSlateStatsFlatIntervalWindowSec;
 							return GSlateStatsFlatIntervalWindowSec;
 						})
 						.MinValue(0.1f)
@@ -568,7 +568,6 @@ void SWidgetReflector::Construct( const FArguments& InArgs )
 						.Delta(0.1f)
 						.OnValueChanged_Static([](float NewValue)
 						{
-							extern SLATECORE_API float GSlateStatsFlatIntervalWindowSec;
 							GSlateStatsFlatIntervalWindowSec = NewValue;
 						})
 					]
@@ -638,7 +637,6 @@ void SWidgetReflector::Construct( const FArguments& InArgs )
 #if SLATE_STATS			
 	if (StatsToolsBox.IsValid())
 	{
-		extern SLATECORE_API int32 GSlateStatsFlatEnable;
 		StatsToolsBox->SetVisibility(GSlateStatsFlatEnable != 0 ? EVisibility::Visible : EVisibility::Collapsed);
 	}
 #endif
@@ -988,7 +986,6 @@ TSharedRef<SWidget> SWidgetReflector::MakeStatViewer()
 		.BorderImage( FCoreStyle::Get().GetBrush("NoBrush") )
 		.Visibility_Lambda([]
 		{
-			extern SLATECORE_API int32 GSlateStatsFlatEnable;
 			return GSlateStatsFlatEnable > 0 ? EVisibility::Visible : EVisibility::Collapsed;
 		})
 		[
diff --git a/Engine/Source/Runtime/SlateCore/Public/Stats/SlateStats.h b/Engine/Source/Runtime/SlateCore/Public/Stats/SlateStats.h
--- a/Engine/Source/Runtime/SlateCore/Public/Stats/SlateStats.h
+++ b/Engine/Source/Runtime/SlateCore/Public/Stats/SlateStats.h
@@ -3,7 +3,7 @@
 #pragma once
 
 /** Set to zero to disable the slate stats system completely, meaning all stats-related structures and code are compiled out. */
-#define SLATE_STATS !UE_BUILD_SHIPPING || WITH_EDITOR
+#define SLATE_STATS (!UE_BUILD_SHIPPING || WITH_EDITOR)
 
 /** 
  * Predefined detail levels to make choosing level slightly easier.
