[INTEGRATE] Change 2444700 by [REDACTED_PERSON][REDACTED_HANDLE] on 2015/02/13 09:56:08

	UE-9790 - Fixing widget reflector to not crash on UMG widgets that have no associated blueprint asset.

diff --git a/Engine/Source/Developer/SlateReflector/Private/Widgets/SWidgetReflectorTreeWidgetItem.cpp b/Engine/Source/Developer/SlateReflector/Private/Widgets/SWidgetReflectorTreeWidgetItem.cpp
--- a/Engine/Source/Developer/SlateReflector/Private/Widgets/SWidgetReflectorTreeWidgetItem.cpp
+++ b/Engine/Source/Developer/SlateReflector/Private/Widgets/SWidgetReflectorTreeWidgetItem.cpp
@@ -103,7 +103,7 @@ FString SReflectorTreeWidgetItem::GetReadableLocation() const
 	{
 		TSharedPtr<SWidget> SafeWidget = Widget.Pin();
 		TSharedPtr<FReflectionMetaData> MetaData = SafeWidget->GetMetaData<FReflectionMetaData>();
-		if ( MetaData.IsValid() )
+		if ( MetaData.IsValid() && MetaData->Asset.Get() != nullptr )
 		{
 			return MetaData->Asset->GetName() + TEXT(" [") + MetaData->Name.ToString() + TEXT("]");
 		}
@@ -119,7 +119,7 @@ void SReflectorTreeWidgetItem::HandleHyperlinkNavigate()
 	{
 		TSharedPtr<SWidget> SafeWidget = Widget.Pin();
 		TSharedPtr<FReflectionMetaData> MetaData = SafeWidget->GetMetaData<FReflectionMetaData>();
-		if ( MetaData.IsValid() )
+		if ( MetaData.IsValid() && MetaData->Asset.Get() != nullptr )
 		{
 			if ( OnAccessAsset.IsBound() )
 			{
