Converted Slate, SlateCore, and SlateViewer to exclusively use FText with Slate

UETOOL-213 - Minimize Slate FString -> FText conversion (remove SLATE_TEXT_ATTRIBUTE)

diff --git a/Engine/Source/Developer/SlateReflector/Private/Widgets/SWidgetReflector.cpp b/Engine/Source/Developer/SlateReflector/Private/Widgets/SWidgetReflector.cpp
--- a/Engine/Source/Developer/SlateReflector/Private/Widgets/SWidgetReflector.cpp
+++ b/Engine/Source/Developer/SlateReflector/Private/Widgets/SWidgetReflector.cpp
@@ -189,7 +189,7 @@ private:
 	}
 
 	/** Callback for getting the text of the pick button. */
-	FString HandlePickButtonText() const;
+	FText HandlePickButtonText() const;
 
 	/** Callback for generating a row in the reflector tree view. */
 	TSharedRef<ITableRow> HandleReflectorTreeGenerateRow( TSharedPtr<FReflectorNode> InReflectorNode, const TSharedRef<STableViewBase>& OwnerTable );
@@ -282,7 +282,7 @@ void SWidgetReflector::Construct( const FArguments& InArgs )
 							.VAlign(VAlign_Center)
 							[
 								SNew(STextBlock)
-									.Text(LOCTEXT("AppScale", "Application Scale: ").ToString())
+									.Text(LOCTEXT("AppScale", "Application Scale: "))
 							]
 
 						+ SHorizontalBox::Slot()
@@ -339,7 +339,7 @@ void SWidgetReflector::Construct( const FArguments& InArgs )
 							.OnCheckStateChanged(this, &SWidgetReflector::HandleFocusCheckBoxCheckedStateChanged)
 							[
 								SNew(STextBlock)
-									.Text(LOCTEXT("ShowFocus", "Show Focus").ToString())
+									.Text(LOCTEXT("ShowFocus", "Show Focus"))
 							]
 					]
 
@@ -545,7 +545,7 @@ void SWidgetReflector::Construct( const FArguments& InArgs )
 								.Text( LOCTEXT("StatsSamplingIntervalTooltip", "the interval (in seconds) to integrate stats before updating the averages.") )
 							]
 						)
-						.Text(LOCTEXT("StatsSampleWindow", "Sampling Interval: ").ToString())
+						.Text(LOCTEXT("StatsSampleWindow", "Sampling Interval: "))
 					]
 					+ SHorizontalBox::Slot()
 					[
@@ -588,7 +588,7 @@ void SWidgetReflector::Construct( const FArguments& InArgs )
 							SNew(SHeaderRow)
 
 							+ SHeaderRow::Column("WidgetName")
-							.DefaultLabel(LOCTEXT("WidgetName", "Widget Name").ToString())
+							.DefaultLabel(LOCTEXT("WidgetName", "Widget Name"))
 							.FillWidth(0.65f)
 
 							+ SHeaderRow::Column("ForegroundColor")
@@ -597,16 +597,16 @@ void SWidgetReflector::Construct( const FArguments& InArgs )
 							.HeaderContent()
 							[
 								SNew(STextBlock)
-								.Text(LOCTEXT("ForegroundColor", "FG").ToString())
-								.ToolTipText(LOCTEXT("ForegroundColorToolTip", "Foreground Color").ToString())
+								.Text(LOCTEXT("ForegroundColor", "FG"))
+								.ToolTipText(LOCTEXT("ForegroundColorToolTip", "Foreground Color"))
 							]
 
 							+ SHeaderRow::Column("Visibility")
-							.DefaultLabel(LOCTEXT("Visibility", "Visibility" ).ToString())
+							.DefaultLabel(LOCTEXT("Visibility", "Visibility" ))
 							.FixedWidth(125.0f)
 
 							+ SHeaderRow::Column("WidgetInfo")
-							.DefaultLabel(LOCTEXT("WidgetInfo", "Widget Info" ).ToString())
+							.DefaultLabel(LOCTEXT("WidgetInfo", "Widget Info" ))
 							.FillWidth(0.25f)
 
 							+ SHeaderRow::Column("Address")
@@ -902,10 +902,10 @@ FString SWidgetReflector::HandleFrameRateText() const
 }
 
 
-FString SWidgetReflector::HandlePickButtonText() const
+FText SWidgetReflector::HandlePickButtonText() const
 {
-	static const FString NotPicking = LOCTEXT("PickWidget", "Pick Widget").ToString();
-	static const FString Picking = LOCTEXT("PickingWidget", "Picking (Esc to Stop)").ToString();
+	static const FText NotPicking = LOCTEXT("PickWidget", "Pick Widget");
+	static const FText Picking = LOCTEXT("PickingWidget", "Picking (Esc to Stop)");
 
 	return bIsPicking ? Picking : NotPicking;
 }
diff --git a/Engine/Source/Developer/SlateReflector/Private/Widgets/SWidgetReflectorToolTipWidget.h b/Engine/Source/Developer/SlateReflector/Private/Widgets/SWidgetReflectorToolTipWidget.h
--- a/Engine/Source/Developer/SlateReflector/Private/Widgets/SWidgetReflectorToolTipWidget.h
+++ b/Engine/Source/Developer/SlateReflector/Private/Widgets/SWidgetReflectorToolTipWidget.h
@@ -95,36 +95,39 @@ public:
 
 private:
 
-	FString GetWidgetsDesiredSize() const
+	FText GetWidgetsDesiredSize() const
 	{
 		TSharedPtr<SWidget> TheWidget = WidgetInfo.Get()->Widget.Pin();
 
 		return TheWidget.IsValid()
-			? TheWidget->GetDesiredSize().ToString()
-			: FString();
+			? FText::FromString(TheWidget->GetDesiredSize().ToString())
+			: FText::GetEmpty();
 	}
 
-	FString GetWidgetActualSize() const
+	FText GetWidgetActualSize() const
 	{
-		return WidgetInfo.Get()->Geometry.Size.ToString();
+		return FText::FromString(WidgetInfo.Get()->Geometry.Size.ToString());
 	}
 
-	FString GetSizeInfo() const
+	FText GetSizeInfo() const
 	{
 		TSharedPtr<SWidget> TheWidget = WidgetInfo.Get()->Widget.Pin();
 
 		return TheWidget.IsValid()
-			? WidgetInfo.Get()->Geometry.ToString()
-			: FString();
+			? FText::FromString(WidgetInfo.Get()->Geometry.ToString())
+			: FText::GetEmpty();
 	}
 
-	FString GetEnabled() const
+	FText GetEnabled() const
 	{
+		static const FText TrueText = LOCTEXT("True", "True");
+		static const FText FalseText = LOCTEXT("False", "False");
+
 		TSharedPtr<SWidget> TheWidget = WidgetInfo.Get()->Widget.Pin();
 
 		return TheWidget.IsValid()
-			? (TheWidget->IsEnabled() ? "True" : "False")
-			: FString();
+			? (TheWidget->IsEnabled() ? TrueText : FalseText)
+			: FText::GetEmpty();
 	}
 
 private:
diff --git a/Engine/Source/Runtime/AppFramework/Private/Framework/Testing/SWidgetGallery.cpp b/Engine/Source/Runtime/AppFramework/Private/Framework/Testing/SWidgetGallery.cpp
--- a/Engine/Source/Runtime/AppFramework/Private/Framework/Testing/SWidgetGallery.cpp
+++ b/Engine/Source/Runtime/AppFramework/Private/Framework/Testing/SWidgetGallery.cpp
@@ -985,7 +985,7 @@ private:
 	TSharedRef<SWidget> HandleComboBoxGenerateWidget( TSharedPtr<FString> InItem )
 	{
 		return SNew(STextBlock)
-			.Text(*InItem);
+			.Text(FText::FromString(*InItem));
 	}
 
 	// Callback for getting the percent value in the SProgressBar example.
@@ -1029,9 +1029,9 @@ private:
 	}
 
 	// Callback for getting the text of the second combo box in the SComboBox example.
-	FString HandleSecondComboBoxText() const
+	FText HandleSecondComboBoxText() const
 	{
-		return ComboString.IsValid() ? *ComboString : FString();
+		return ComboString.IsValid() ? FText::FromString(*ComboString) : FText::GetEmpty();
 	}
 
 	// Callback for changing the selector combo box's selection in SComboBox example.
@@ -1050,9 +1050,9 @@ private:
 	}
 
 	// Callback for getting the text of the selector combo box in the SComboBox example.
-	FString HandleSelectorComboBoxText( ) const
+	FText HandleSelectorComboBoxText( ) const
 	{
-		return SelectorComboBoxSelectedItem.IsValid() ? *SelectorComboBoxSelectedItem : FString();
+		return SelectorComboBoxSelectedItem.IsValid() ? FText::FromString(*SelectorComboBoxSelectedItem) : FText::GetEmpty();
 	}
 
 	// Callback for changing the active tab.
diff --git a/Engine/Source/Runtime/AppFramework/Private/Widgets/Colors/SColorThemes.cpp b/Engine/Source/Runtime/AppFramework/Private/Widgets/Colors/SColorThemes.cpp
--- a/Engine/Source/Runtime/AppFramework/Private/Widgets/Colors/SColorThemes.cpp
+++ b/Engine/Source/Runtime/AppFramework/Private/Widgets/Colors/SColorThemes.cpp
@@ -739,9 +739,9 @@ void SColorThemeBar::Construct(const FArguments& InArgs)
 	];
 }
 
-FString SColorThemeBar::GetThemeName() const
+FText SColorThemeBar::GetThemeName() const
 {
-	return ColorTheme.Pin()->Name;
+	return FText::FromString(ColorTheme.Pin()->Name);
 }
 
 FReply SColorThemeBar::OnMouseButtonDown( const FGeometry& MyGeometry, const FPointerEvent& MouseEvent )
diff --git a/Engine/Source/Runtime/AppFramework/Private/Widgets/Testing/STestSuite.cpp b/Engine/Source/Runtime/AppFramework/Private/Widgets/Testing/STestSuite.cpp
--- a/Engine/Source/Runtime/AppFramework/Private/Widgets/Testing/STestSuite.cpp
+++ b/Engine/Source/Runtime/AppFramework/Private/Widgets/Testing/STestSuite.cpp
@@ -3897,7 +3897,7 @@ protected:
 	void UpdateColor(FLinearColor NewColor)
 	{
 		Color.Set(NewColor);
-		OutputTextBlock->SetText( Color.Get().ToFColor(false).ToString() );
+		OutputTextBlock->SetText( FText::FromString(Color.Get().ToFColor(false).ToString()) );
 	}
 	
 	TAttribute<FLinearColor> Color;
@@ -4438,7 +4438,7 @@ TSharedRef<SWidget> SResponsiveGridPanelTestWidgetImpl::ConstructBox(const FStri
 			.Padding(10)
 			.BorderImage(FTestStyle::Get().GetBrush("RichText.Tagline.Background"))
 			[
-				SNew(STextBlock).Text(Text)
+				SNew(STextBlock).Text(FText::FromString(Text))
 			]
 		];
 }
diff --git a/Engine/Source/Runtime/AppFramework/Public/Widgets/Colors/SColorThemes.h b/Engine/Source/Runtime/AppFramework/Public/Widgets/Colors/SColorThemes.h
--- a/Engine/Source/Runtime/AppFramework/Public/Widgets/Colors/SColorThemes.h
+++ b/Engine/Source/Runtime/AppFramework/Public/Widgets/Colors/SColorThemes.h
@@ -459,7 +459,7 @@ public:
 	
 private:
 
-	FString GetThemeName() const;
+	FText GetThemeName() const;
 
 	/** Text Block which shows the Color Theme's name */
 	TSharedPtr<STextBlock> ThemeNameText;
diff --git a/Engine/Source/Runtime/Slate/Private/Framework/MultiBox/MultiBox.cpp b/Engine/Source/Runtime/Slate/Private/Framework/MultiBox/MultiBox.cpp
--- a/Engine/Source/Runtime/Slate/Private/Framework/MultiBox/MultiBox.cpp
+++ b/Engine/Source/Runtime/Slate/Private/Framework/MultiBox/MultiBox.cpp
@@ -491,7 +491,7 @@ void SMultiBoxWidget::AddBlockWidget( const FMultiBlock& Block, TSharedPtr<SHori
 					SNew(STextBlock)
 					.Visibility(bDisplayExtensionHooks ? EVisibility::Visible : EVisibility::Collapsed)
 					.ColorAndOpacity(StyleSet->GetColor("MultiboxHookColor"))
-					.Text(Block.GetExtensionHook().ToString())
+					.Text(FText::FromName(Block.GetExtensionHook()))
 				]
 				+SVerticalBox::Slot()
 				[
@@ -514,7 +514,7 @@ void SMultiBoxWidget::AddBlockWidget( const FMultiBlock& Block, TSharedPtr<SHori
 						SNew(STextBlock)
 						.Visibility(bDisplayExtensionHooks ? EVisibility::Visible : EVisibility::Collapsed)
 						.ColorAndOpacity(StyleSet->GetColor("MultiboxHookColor"))
-						.Text(Block.GetExtensionHook().ToString())
+						.Text(FText::FromName(Block.GetExtensionHook()))
 					]
 					+SVerticalBox::Slot()
 					[
@@ -542,7 +542,7 @@ void SMultiBoxWidget::AddBlockWidget( const FMultiBlock& Block, TSharedPtr<SHori
 					SNew(STextBlock)
 					.Visibility(bDisplayExtensionHooks ? EVisibility::Visible : EVisibility::Collapsed)
 					.ColorAndOpacity(StyleSet->GetColor("MultiboxHookColor"))
-					.Text(Block.GetExtensionHook().ToString())
+					.Text(FText::FromName(Block.GetExtensionHook()))
 				]
 				+SHorizontalBox::Slot()
 				[
diff --git a/Engine/Source/Runtime/Slate/Private/Widgets/Input/SSuggestionTextBox.cpp b/Engine/Source/Runtime/Slate/Private/Widgets/Input/SSuggestionTextBox.cpp
--- a/Engine/Source/Runtime/Slate/Private/Widgets/Input/SSuggestionTextBox.cpp
+++ b/Engine/Source/Runtime/Slate/Private/Widgets/Input/SSuggestionTextBox.cpp
@@ -311,7 +311,7 @@ TSharedRef<ITableRow> SSuggestionTextBox::HandleSuggestionListViewGenerateRow( T
 					SNew(STextBlock)
 						.HighlightText(this, &SSuggestionTextBox::HandleSuggestionListWidgetHighlightText)
 						.TextStyle( SuggestionTextStyle )
-						.Text( Combined )						
+						.Text( FText::FromString(Combined) )						
 				]
 		];
 }
diff --git a/Engine/Source/Runtime/Slate/Private/Widgets/Input/STextComboBox.cpp b/Engine/Source/Runtime/Slate/Private/Widgets/Input/STextComboBox.cpp
--- a/Engine/Source/Runtime/Slate/Private/Widgets/Input/STextComboBox.cpp
+++ b/Engine/Source/Runtime/Slate/Private/Widgets/Input/STextComboBox.cpp
@@ -32,19 +32,22 @@ void STextComboBox::Construct( const FArguments& InArgs )
 	SelectedItem = StringCombo->GetSelectedItem();
 }
 
-FString STextComboBox::GetItemTextLabel(TSharedPtr<FString> StringItem) const
+FText STextComboBox::GetItemTextLabel(TSharedPtr<FString> StringItem) const
 {
 	if (!StringItem.IsValid())
 	{
-		return FString();
+		return FText::GetEmpty();
 	}
 
-	return (GetTextLabelForItem.IsBound())
+	// todo: jdale - should this be using FText natively?
+	return FText::FromString(
+		(GetTextLabelForItem.IsBound())
 		? GetTextLabelForItem.Execute(StringItem)
-		: *StringItem;
+		: *StringItem
+		);
 }
 
-FString STextComboBox::GetSelectedTextLabel() const
+FText STextComboBox::GetSelectedTextLabel() const
 {
 	TSharedPtr<FString> StringItem = StringCombo->GetSelectedItem();
 	return GetItemTextLabel(StringItem);
diff --git a/Engine/Source/Runtime/Slate/Private/Widgets/Input/STextComboPopup.cpp b/Engine/Source/Runtime/Slate/Private/Widgets/Input/STextComboPopup.cpp
--- a/Engine/Source/Runtime/Slate/Private/Widgets/Input/STextComboPopup.cpp
+++ b/Engine/Source/Runtime/Slate/Private/Widgets/Input/STextComboPopup.cpp
@@ -63,14 +63,16 @@ void STextComboPopup::Construct( const FArguments& InArgs )
 
 TSharedRef<SWidget> STextComboPopup::MakeItemWidget( TSharedPtr<FString> StringItem ) 
 {
+	// todo: jdale - should this be using FText natively?
 	check( StringItem.IsValid() );
 	return	SNew(STextBlock)
-		.Text(*StringItem);
+		.Text(FText::FromString(*StringItem));
 }
 
-FString STextComboPopup::GetSelectedItem() const
+FText STextComboPopup::GetSelectedItem() const
 {
-	return *SelectedItem;
+	// todo: jdale - should this be using FText natively?
+	return FText::FromString(*SelectedItem);
 }
 
 void STextComboPopup::OnSelectionChanged (TSharedPtr<FString> Selection, ESelectInfo::Type SelectInfo)
diff --git a/Engine/Source/Runtime/Slate/Public/Widgets/Input/SComboBox.h b/Engine/Source/Runtime/Slate/Public/Widgets/Input/SComboBox.h
--- a/Engine/Source/Runtime/Slate/Public/Widgets/Input/SComboBox.h
+++ b/Engine/Source/Runtime/Slate/Public/Widgets/Input/SComboBox.h
@@ -167,7 +167,7 @@ public:
 		if (InArgs._Content.Widget == SNullWidget::NullWidget)
 		{
 			 SAssignNew(ButtonContent, STextBlock)
-			.Text(NSLOCTEXT("SComboBox", "ContentWarning", "No Content Provided").ToString())
+			.Text(NSLOCTEXT("SComboBox", "ContentWarning", "No Content Provided"))
 			.ColorAndOpacity( FLinearColor::Red);
 		}
 
diff --git a/Engine/Source/Runtime/Slate/Public/Widgets/Input/SSpinBox.h b/Engine/Source/Runtime/Slate/Public/Widgets/Input/SSpinBox.h
--- a/Engine/Source/Runtime/Slate/Public/Widgets/Input/SSpinBox.h
+++ b/Engine/Source/Runtime/Slate/Public/Widgets/Input/SSpinBox.h
@@ -141,7 +141,7 @@ public:
 			[
 				SAssignNew(TextBlock, STextBlock)
 				.Font(InArgs._Font)
-				.Text( this, &SSpinBox<NumericType>::GetValueAsString )
+				.Text( this, &SSpinBox<NumericType>::GetValueAsText )
 				.MinDesiredWidth( this, &SSpinBox<NumericType>::GetTextMinDesiredWidth )
 			]
 
diff --git a/Engine/Source/Runtime/Slate/Public/Widgets/Input/STextComboBox.h b/Engine/Source/Runtime/Slate/Public/Widgets/Input/STextComboBox.h
--- a/Engine/Source/Runtime/Slate/Public/Widgets/Input/STextComboBox.h
+++ b/Engine/Source/Runtime/Slate/Public/Widgets/Input/STextComboBox.h
@@ -76,9 +76,9 @@ private:
 	void OnSelectionChanged (TSharedPtr<FString> Selection, ESelectInfo::Type SelectInfo);
 
 	/** Helper method to get the text for a given item in the combo box */
-	FString GetSelectedTextLabel() const;
+	FText GetSelectedTextLabel() const;
 
-	FString GetItemTextLabel(TSharedPtr<FString> StringItem) const;
+	FText GetItemTextLabel(TSharedPtr<FString> StringItem) const;
 
 private:
 
diff --git a/Engine/Source/Runtime/Slate/Public/Widgets/Input/STextComboPopup.h b/Engine/Source/Runtime/Slate/Public/Widgets/Input/STextComboPopup.h
--- a/Engine/Source/Runtime/Slate/Public/Widgets/Input/STextComboPopup.h
+++ b/Engine/Source/Runtime/Slate/Public/Widgets/Input/STextComboPopup.h
@@ -52,7 +52,7 @@ private:
 	void OnSelectionChanged (TSharedPtr<FString> Selection, ESelectInfo::Type SelectInfo);
 
 	/** Returns the item to select in the combo popup */
-	FString GetSelectedItem () const;
+	FText GetSelectedItem () const;
 
 private:
 	/** The text item selected */
diff --git a/Engine/Source/Runtime/Slate/Public/Widgets/Views/SListView.h b/Engine/Source/Runtime/Slate/Public/Widgets/Views/SListView.h
--- a/Engine/Source/Runtime/Slate/Public/Widgets/Views/SListView.h
+++ b/Engine/Source/Runtime/Slate/Public/Widgets/Views/SListView.h
@@ -145,7 +145,7 @@ public:
 				.VAlign(VAlign_Center)
 				[
 					SNew(STextBlock)
-					.Text(ErrorString)
+					.Text(FText::FromString(ErrorString))
 				];
 		}
 		else
diff --git a/Engine/Source/Runtime/Slate/Public/Widgets/Views/STileView.h b/Engine/Source/Runtime/Slate/Public/Widgets/Views/STileView.h
--- a/Engine/Source/Runtime/Slate/Public/Widgets/Views/STileView.h
+++ b/Engine/Source/Runtime/Slate/Public/Widgets/Views/STileView.h
@@ -112,7 +112,7 @@ public:
 			.VAlign(VAlign_Center)
 			[
 				SNew(STextBlock)
-				.Text(ErrorString)
+				.Text(FText::FromString(ErrorString))
 			];
 		}
 		else
diff --git a/Engine/Source/Runtime/Slate/Public/Widgets/Views/STreeView.h b/Engine/Source/Runtime/Slate/Public/Widgets/Views/STreeView.h
--- a/Engine/Source/Runtime/Slate/Public/Widgets/Views/STreeView.h
+++ b/Engine/Source/Runtime/Slate/Public/Widgets/Views/STreeView.h
@@ -172,7 +172,7 @@ public:
 			.VAlign(VAlign_Center)
 			[
 				SNew(STextBlock)
-				.Text(ErrorString)
+				.Text(FText::FromString(ErrorString))
 			];
 		}
 		else
