Fixed an issue where the reflector module had a variable to cache off the last created reflector widget but never assigned to it.
Merging using UE4-To-UE4-LauncherDev

diff --git a/Engine/Source/Developer/SlateReflector/Private/SlateReflectorModule.cpp b/Engine/Source/Developer/SlateReflector/Private/SlateReflectorModule.cpp
--- a/Engine/Source/Developer/SlateReflector/Private/SlateReflectorModule.cpp
+++ b/Engine/Source/Developer/SlateReflector/Private/SlateReflectorModule.cpp
@@ -27,6 +27,7 @@ public:
 		if (!WidgetReflector.IsValid())
 		{
 			WidgetReflector = SNew(SWidgetReflector);
+			WidgetReflectorPtr = WidgetReflector;
 			FSlateApplication::Get().SetWidgetReflector(WidgetReflector.ToSharedRef());
 		}
 
