Epic Friday - Working on adding "snapshots" to the widget reflector

The underlying changes have been made to allow the widget reflector to visualize a live or snapshotted tree.

Submitting these now as the API has changed quite a bit and any future changes will likely conflict.

The "Snapshot" button is implemented but is currently commented out as it's not fully functional (I still need to take a screenshot of the windows so you can pick against their saved state).

diff --git a/Engine/Source/Developer/SlateReflector/Private/Models/WidgetReflectorNode.cpp b/Engine/Source/Developer/SlateReflector/Private/Models/WidgetReflectorNode.cpp
new file mode 100644
--- /dev/null
+++ b/Engine/Source/Developer/SlateReflector/Private/Models/WidgetReflectorNode.cpp
@@ -0,0 +1,386 @@
+// Copyright 1998-2015 Epic Games, Inc. All Rights Reserved.
+
+#include "SlateReflectorPrivatePCH.h"
+#include "WidgetReflectorNode.h"
+#include "ReflectionMetadata.h"
+
+#define LOCTEXT_NAMESPACE "WidgetReflectorNode"
+
+/**
+ * -----------------------------------------------------------------------------
+ * UWidgetReflectorNodeBase
+ * -----------------------------------------------------------------------------
+ */
+void UWidgetReflectorNodeBase::Initialize(const FArrangedWidget& InWidgetGeometry)
+{
+	Geometry = InWidgetGeometry.Geometry;
+	Tint = FLinearColor(1.0f, 1.0f, 1.0f);
+}
+
+TSharedPtr<SWidget> UWidgetReflectorNodeBase::GetLiveWidget() const
+{
+	return nullptr;
+}
+
+FText UWidgetReflectorNodeBase::GetWidgetType() const
+{
+	return FText::GetEmpty();
+}
+
+FText UWidgetReflectorNodeBase::GetWidgetVisibilityText() const
+{
+	return FText::GetEmpty();
+}
+
+FText UWidgetReflectorNodeBase::GetWidgetReadableLocation() const
+{
+	return FText::GetEmpty();
+}
+
+FString UWidgetReflectorNodeBase::GetWidgetFile() const
+{
+	return FString();
+}
+
+int32 UWidgetReflectorNodeBase::GetWidgetLineNumber() const
+{
+	return 0;
+}
+
+FName UWidgetReflectorNodeBase::GetWidgetAssetName() const
+{
+	return NAME_None;
+}
+
+FVector2D UWidgetReflectorNodeBase::GetWidgetDesiredSize() const
+{
+	return FVector2D::ZeroVector;
+}
+
+FSlateColor UWidgetReflectorNodeBase::GetWidgetForegroundColor() const
+{
+	return FSlateColor::UseForeground();
+}
+
+FString UWidgetReflectorNodeBase::GetWidgetAddress() const
+{
+	return FString::Printf(TEXT("0x%p"), static_cast<const void*>(nullptr));
+}
+
+bool UWidgetReflectorNodeBase::GetWidgetEnabled() const
+{
+	return false;
+}
+
+const FGeometry& UWidgetReflectorNodeBase::GetGeometry() const
+{
+	return Geometry;
+}
+
+const FLinearColor& UWidgetReflectorNodeBase::GetTint() const
+{
+	return Tint;
+}
+
+void UWidgetReflectorNodeBase::SetTint(const FLinearColor& InTint)
+{
+	Tint = InTint;
+}
+
+void UWidgetReflectorNodeBase::AddChildNode(UWidgetReflectorNodeBase* InChildNode)
+{
+	ChildNodes.Add(InChildNode);
+}
+
+const TArray<UWidgetReflectorNodeBase*>& UWidgetReflectorNodeBase::GetChildNodes() const
+{
+	return ChildNodes;
+}
+
+
+/**
+ * -----------------------------------------------------------------------------
+ * ULiveWidgetReflectorNode
+ * -----------------------------------------------------------------------------
+ */
+void ULiveWidgetReflectorNode::Initialize(const FArrangedWidget& InWidgetGeometry)
+{
+	Super::Initialize(InWidgetGeometry);
+
+	Widget = InWidgetGeometry.Widget;
+}
+
+TSharedPtr<SWidget> ULiveWidgetReflectorNode::GetLiveWidget() const
+{
+	return Widget.Pin();
+}
+
+FText ULiveWidgetReflectorNode::GetWidgetType() const
+{
+	return FWidgetReflectorNodeUtils::GetWidgetType(Widget.Pin());
+}
+
+FText ULiveWidgetReflectorNode::GetWidgetVisibilityText() const
+{
+	return FWidgetReflectorNodeUtils::GetWidgetVisibilityText(Widget.Pin());
+}
+
+FText ULiveWidgetReflectorNode::GetWidgetReadableLocation() const
+{
+	return FWidgetReflectorNodeUtils::GetWidgetReadableLocation(Widget.Pin());
+}
+
+FString ULiveWidgetReflectorNode::GetWidgetFile() const
+{
+	return FWidgetReflectorNodeUtils::GetWidgetFile(Widget.Pin());
+}
+
+int32 ULiveWidgetReflectorNode::GetWidgetLineNumber() const
+{
+	return FWidgetReflectorNodeUtils::GetWidgetLineNumber(Widget.Pin());
+}
+
+FName ULiveWidgetReflectorNode::GetWidgetAssetName() const
+{
+	return FWidgetReflectorNodeUtils::GetWidgetAssetName(Widget.Pin());
+}
+
+FVector2D ULiveWidgetReflectorNode::GetWidgetDesiredSize() const
+{
+	return FWidgetReflectorNodeUtils::GetWidgetDesiredSize(Widget.Pin());
+}
+
+FSlateColor ULiveWidgetReflectorNode::GetWidgetForegroundColor() const
+{
+	return FWidgetReflectorNodeUtils::GetWidgetForegroundColor(Widget.Pin());
+}
+
+FString ULiveWidgetReflectorNode::GetWidgetAddress() const
+{
+	return FWidgetReflectorNodeUtils::GetWidgetAddress(Widget.Pin());
+}
+
+bool ULiveWidgetReflectorNode::GetWidgetEnabled() const
+{
+	return FWidgetReflectorNodeUtils::GetWidgetEnabled(Widget.Pin());
+}
+
+
+/**
+ * -----------------------------------------------------------------------------
+ * USnapshotWidgetReflectorNode
+ * -----------------------------------------------------------------------------
+ */
+void USnapshotWidgetReflectorNode::Initialize(const FArrangedWidget& InWidgetGeometry)
+{
+	Super::Initialize(InWidgetGeometry);
+
+	CachedWidgetType = FWidgetReflectorNodeUtils::GetWidgetType(InWidgetGeometry.Widget);
+	CachedWidgetVisibilityText = FWidgetReflectorNodeUtils::GetWidgetVisibilityText(InWidgetGeometry.Widget);
+	CachedWidgetReadableLocation = FWidgetReflectorNodeUtils::GetWidgetReadableLocation(InWidgetGeometry.Widget);
+	CachedWidgetFile = FWidgetReflectorNodeUtils::GetWidgetFile(InWidgetGeometry.Widget);
+	CachedWidgetLineNumber = FWidgetReflectorNodeUtils::GetWidgetLineNumber(InWidgetGeometry.Widget);
+	CachedWidgetAssetName = FWidgetReflectorNodeUtils::GetWidgetAssetName(InWidgetGeometry.Widget);
+	CachedWidgetDesiredSize = FWidgetReflectorNodeUtils::GetWidgetDesiredSize(InWidgetGeometry.Widget);
+	CachedWidgetForegroundColor = FWidgetReflectorNodeUtils::GetWidgetForegroundColor(InWidgetGeometry.Widget);
+	CachedWidgetAddress = FWidgetReflectorNodeUtils::GetWidgetAddress(InWidgetGeometry.Widget);
+	CachedWidgetEnabled = FWidgetReflectorNodeUtils::GetWidgetEnabled(InWidgetGeometry.Widget);
+}
+
+FText USnapshotWidgetReflectorNode::GetWidgetType() const
+{
+	return CachedWidgetType;
+}
+
+FText USnapshotWidgetReflectorNode::GetWidgetVisibilityText() const
+{
+	return CachedWidgetVisibilityText;
+}
+
+FText USnapshotWidgetReflectorNode::GetWidgetReadableLocation() const
+{
+	return CachedWidgetReadableLocation;
+}
+
+FString USnapshotWidgetReflectorNode::GetWidgetFile() const
+{
+	return CachedWidgetFile;
+}
+
+int32 USnapshotWidgetReflectorNode::GetWidgetLineNumber() const
+{
+	return CachedWidgetLineNumber;
+}
+
+FName USnapshotWidgetReflectorNode::GetWidgetAssetName() const
+{
+	return CachedWidgetAssetName;
+}
+
+FVector2D USnapshotWidgetReflectorNode::GetWidgetDesiredSize() const
+{
+	return CachedWidgetDesiredSize;
+}
+
+FSlateColor USnapshotWidgetReflectorNode::GetWidgetForegroundColor() const
+{
+	return CachedWidgetForegroundColor;
+}
+
+FString USnapshotWidgetReflectorNode::GetWidgetAddress() const
+{
+	return CachedWidgetAddress;
+}
+
+bool USnapshotWidgetReflectorNode::GetWidgetEnabled() const
+{
+	return CachedWidgetEnabled;
+}
+
+
+/**
+ * -----------------------------------------------------------------------------
+ * FWidgetReflectorNodeUtils
+ * -----------------------------------------------------------------------------
+ */
+ULiveWidgetReflectorNode* FWidgetReflectorNodeUtils::NewLiveNode(const FArrangedWidget& InWidgetGeometry)
+{
+	return Cast<ULiveWidgetReflectorNode>(NewNode(ULiveWidgetReflectorNode::StaticClass(), InWidgetGeometry));
+}
+
+ULiveWidgetReflectorNode* FWidgetReflectorNodeUtils::NewLiveNodeTreeFrom(const FArrangedWidget& InWidgetGeometry)
+{
+	return Cast<ULiveWidgetReflectorNode>(NewNodeTreeFrom(ULiveWidgetReflectorNode::StaticClass(), InWidgetGeometry));
+}
+
+USnapshotWidgetReflectorNode* FWidgetReflectorNodeUtils::NewSnapshotNode(const FArrangedWidget& InWidgetGeometry)
+{
+	return Cast<USnapshotWidgetReflectorNode>(NewNode(USnapshotWidgetReflectorNode::StaticClass(), InWidgetGeometry));
+}
+
+USnapshotWidgetReflectorNode* FWidgetReflectorNodeUtils::NewSnapshotNodeTreeFrom(const FArrangedWidget& InWidgetGeometry)
+{
+	return Cast<USnapshotWidgetReflectorNode>(NewNodeTreeFrom(USnapshotWidgetReflectorNode::StaticClass(), InWidgetGeometry));
+}
+
+UWidgetReflectorNodeBase* FWidgetReflectorNodeUtils::NewNode(const TSubclassOf<UWidgetReflectorNodeBase>& InNodeClass, const FArrangedWidget& InWidgetGeometry)
+{
+	UWidgetReflectorNodeBase* NewNodeInstance = NewObject<UWidgetReflectorNodeBase>(GetTransientPackage(), *InNodeClass);
+	NewNodeInstance->Initialize(InWidgetGeometry);
+	return NewNodeInstance;
+}
+
+UWidgetReflectorNodeBase* FWidgetReflectorNodeUtils::NewNodeTreeFrom(const TSubclassOf<UWidgetReflectorNodeBase>& InNodeClass, const FArrangedWidget& InWidgetGeometry)
+{
+	UWidgetReflectorNodeBase* NewNodeInstance = NewNode(InNodeClass, InWidgetGeometry);
+
+	FArrangedChildren ArrangedChildren(EVisibility::All);
+	InWidgetGeometry.Widget->ArrangeChildren(InWidgetGeometry.Geometry, ArrangedChildren);
+	
+	for (int32 WidgetIndex = 0; WidgetIndex < ArrangedChildren.Num(); ++WidgetIndex)
+	{
+		// Note that we include both visible and invisible children!
+		NewNodeInstance->AddChildNode(NewNodeTreeFrom(InNodeClass, ArrangedChildren[WidgetIndex]));
+	}
+
+	return NewNodeInstance;
+}
+
+void FWidgetReflectorNodeUtils::FindLiveWidgetPath(const TArray<UWidgetReflectorNodeBase*>& CandidateNodes, const FWidgetPath& WidgetPathToFind, TArray<UWidgetReflectorNodeBase*>& SearchResult, int32 NodeIndexToFind)
+{
+	if (NodeIndexToFind < WidgetPathToFind.Widgets.Num())
+	{
+		const FArrangedWidget& WidgetToFind = WidgetPathToFind.Widgets[NodeIndexToFind];
+
+		for (int32 NodeIndex = 0; NodeIndex < CandidateNodes.Num(); ++NodeIndex)
+		{
+			if (CandidateNodes[NodeIndex]->GetWidgetAddress() == GetWidgetAddress(WidgetPathToFind.Widgets[NodeIndexToFind].Widget))
+			{
+				SearchResult.Add(CandidateNodes[NodeIndex]);
+				FindLiveWidgetPath(CandidateNodes[NodeIndex]->GetChildNodes(), WidgetPathToFind, SearchResult, NodeIndexToFind + 1);
+			}
+		}
+	}
+}
+
+FText FWidgetReflectorNodeUtils::GetWidgetType(const TSharedPtr<SWidget>& InWidget)
+{
+	return (InWidget.IsValid()) ? FText::FromString(InWidget->GetTypeAsString()) : FText::GetEmpty();
+}
+
+FText FWidgetReflectorNodeUtils::GetWidgetVisibilityText(const TSharedPtr<SWidget>& InWidget)
+{
+	return (InWidget.IsValid()) ? FText::FromString(InWidget->GetVisibility().ToString()) : FText::GetEmpty();
+}
+
+FText FWidgetReflectorNodeUtils::GetWidgetReadableLocation(const TSharedPtr<SWidget>& InWidget)
+{
+	FText ReadableLocationText;
+
+	if (InWidget.IsValid())
+	{
+		// UMG widgets have meta-data to help track them
+		TSharedPtr<FReflectionMetaData> MetaData = InWidget->GetMetaData<FReflectionMetaData>();
+		if (MetaData.IsValid() && MetaData->Asset.Get() != nullptr)
+		{
+			const FText AssetNameText = FText::FromName(MetaData->Asset->GetFName());
+			const FText WidgetNameText = FText::FromName(MetaData->Name);
+			ReadableLocationText = FText::Format(LOCTEXT("AssetReadableLocationFmt", "{0} [{1}]"), AssetNameText, WidgetNameText);
+		}
+		else
+		{
+			ReadableLocationText = FText::FromString(InWidget->GetReadableLocation());
+		}
+	}
+	
+	return ReadableLocationText;
+}
+
+FString FWidgetReflectorNodeUtils::GetWidgetFile(const TSharedPtr<SWidget>& InWidget)
+{
+	return (InWidget.IsValid()) ? InWidget->GetCreatedInLocation().GetPlainNameString() : FString();
+}
+
+int32 FWidgetReflectorNodeUtils::GetWidgetLineNumber(const TSharedPtr<SWidget>& InWidget)
+{
+	return (InWidget.IsValid()) ? InWidget->GetCreatedInLocation().GetNumber() : 0;
+}
+
+FName FWidgetReflectorNodeUtils::GetWidgetAssetName(const TSharedPtr<SWidget>& InWidget)
+{
+	FName AssetName;
+
+	if (InWidget.IsValid())
+	{
+		// UMG widgets have meta-data to help track them
+		TSharedPtr<FReflectionMetaData> MetaData = InWidget->GetMetaData<FReflectionMetaData>();
+		if (MetaData.IsValid() && MetaData->Asset.Get() != nullptr)
+		{
+			AssetName = MetaData->Asset->GetFName();
+		}
+	}
+
+	return AssetName;
+}
+
+FVector2D FWidgetReflectorNodeUtils::GetWidgetDesiredSize(const TSharedPtr<SWidget>& InWidget)
+{
+	return (InWidget.IsValid()) ? InWidget->GetDesiredSize() : FVector2D::ZeroVector;
+}
+
+FString FWidgetReflectorNodeUtils::GetWidgetAddress(const TSharedPtr<SWidget>& InWidget)
+{
+	return FString::Printf(TEXT("0x%p"), InWidget.Get());
+}
+
+FSlateColor FWidgetReflectorNodeUtils::GetWidgetForegroundColor(const TSharedPtr<SWidget>& InWidget)
+{
+	return (InWidget.IsValid()) ? InWidget->GetForegroundColor() : FSlateColor::UseForeground();
+}
+
+bool FWidgetReflectorNodeUtils::GetWidgetEnabled(const TSharedPtr<SWidget>& InWidget)
+{
+	return (InWidget.IsValid()) ? InWidget->IsEnabled() : false;
+}
+
+#undef LOCTEXT_NAMESPACE
diff --git a/Engine/Source/Developer/SlateReflector/Private/Models/WidgetReflectorNode.h b/Engine/Source/Developer/SlateReflector/Private/Models/WidgetReflectorNode.h
--- a/Engine/Source/Developer/SlateReflector/Private/Models/WidgetReflectorNode.h
+++ b/Engine/Source/Developer/SlateReflector/Private/Models/WidgetReflectorNode.h
@@ -2,98 +2,329 @@
 
 #pragma once
 
+#include "WidgetReflectorNode.generated.h"
 
-/**
-  * The WidgetReflector captures layout information and widget hierarchy structure into a tree of FReflectorNode.
-  * The node also contains useful visualization and debug info.
-  */
-struct FReflectorNode
+/** 
+ * A widget reflector node that contains the interface and basic data required by both live and snapshot nodes 
+ */
+UCLASS()
+class UWidgetReflectorNodeBase : public UObject
 {
-	/** The widget which this node describes */
-	TWeakPtr<SWidget> Widget;
+	GENERATED_BODY()
+
+public:
+	/**
+	 * Initialize this node from the given widget geometry, caching out any data that may be required for future visualization in the widget reflector
+	 */
+	virtual void Initialize(const FArrangedWidget& InWidgetGeometry);
+
+	/**
+	 * @note This function only works for ULiveWidgetReflectorNode instances
+	 * @return The live widget that this node is referencing
+	 */
+	virtual TSharedPtr<SWidget> GetLiveWidget() const;
+
+	/**
+	 * @return The type string for the widget we were initialized from
+	 */
+	virtual FText GetWidgetType() const;
+	
+	/**
+	 * @return The visibility string for the widget we were initialized from
+	 */
+	virtual FText GetWidgetVisibilityText() const;
+
+	/**
+	 * The human readable location for widgets that are defined in C++ is the file and line number
+	 * The human readable location for widgets that are defined in UMG is the asset name
+	 * @return The fully human readable location for the widget we were initialized from
+	 */
+	virtual FText GetWidgetReadableLocation() const;
+
+	/**
+	 * @return The name of the file that the widget we were initialized from was created from (for C++ widgets)
+	 */
+	virtual FString GetWidgetFile() const;
+
+	/**
+	 * @return The line number of the file that the widget we were initialized from was created from (for C++ widgets)
+	 */
+	virtual int32 GetWidgetLineNumber() const;
+
+	/**
+	 * @return The name of the asset that the widget we were initialized from was created from (for UMG widgets)
+	 */
+	virtual FName GetWidgetAssetName() const;
+
+	/**
+	 * @return The desired size of the widget we were initialized from
+	 */
+	virtual FVector2D GetWidgetDesiredSize() const;
+
+	/**
+	 * @return The foreground color of the widget we were initialized from
+	 */
+	virtual FSlateColor GetWidgetForegroundColor() const;
+
+	/**
+	 * @return The in-memory address of the widget we were initialized from
+	 */
+	virtual FString GetWidgetAddress() const;
+
+	/**
+	 * @return True if the the widget we were initialized from is enabled, false otherwise
+	 */
+	virtual bool GetWidgetEnabled() const;
+
+	/**
+	 * @return The geometry of the widget we were initialized from
+	 */
+	const FGeometry& GetGeometry() const;
+
+	/**
+	 * @return The tint that is applied to text in order to provide visual hints
+	 */
+	const FLinearColor& GetTint() const;
+
+	/**
+	 * Set the tint to the given value
+	 */
+	void SetTint(const FLinearColor& InTint);
+
+	/**
+	 * Add the given node to our list of children for this widget
+	 */
+	void AddChildNode(UWidgetReflectorNodeBase* InChildNode);
 
+	/**
+	 * @return The node entries for the widget's children
+	 */
+	const TArray<UWidgetReflectorNodeBase*>& GetChildNodes() const;
+
+protected:
 	/** The geometry of the widget */
+	UPROPERTY()
 	FGeometry Geometry;
 
-	/** ReflectorNodes for the Widget's children. */
-	TArray<TSharedPtr<FReflectorNode>> ChildNodes;	
+	/** Node entries for the widget's children */
+	UPROPERTY()
+	TArray<UWidgetReflectorNodeBase*> ChildNodes;
 
 	/** A tint that is applied to text in order to provide visual hints */
+	UPROPERTY()
 	FLinearColor Tint;
+};
+
+/** 
+ * A widget reflector node that holds on to the widget it references so that certain properties can be updated live 
+ */
+UCLASS()
+class ULiveWidgetReflectorNode : public UWidgetReflectorNodeBase
+{
+	GENERATED_BODY()
+
+public:
+	// UWidgetReflectorNodeBase interface
+	virtual void Initialize(const FArrangedWidget& InWidgetGeometry) override;
+	virtual TSharedPtr<SWidget> GetLiveWidget() const override;
+	virtual FText GetWidgetType() const override;
+	virtual FText GetWidgetVisibilityText() const override;
+	virtual FText GetWidgetReadableLocation() const override;
+	virtual FString GetWidgetFile() const override;
+	virtual int32 GetWidgetLineNumber() const override;
+	virtual FName GetWidgetAssetName() const override;
+	virtual FVector2D GetWidgetDesiredSize() const override;
+	virtual FSlateColor GetWidgetForegroundColor() const override;
+	virtual FString GetWidgetAddress() const override;
+	virtual bool GetWidgetEnabled() const override;
+
+private:
+	/** The widget this node is watching */
+	TWeakPtr<SWidget> Widget;
+};
+
+/** 
+ * A widget reflector node that holds the widget information from a snapshot at a given point in time 
+ */
+UCLASS()
+class USnapshotWidgetReflectorNode : public UWidgetReflectorNodeBase
+{
+	GENERATED_BODY()
+
+public:
+	// UWidgetReflectorNodeBase interface
+	virtual void Initialize(const FArrangedWidget& InWidgetGeometry) override;
+	virtual FText GetWidgetType() const override;
+	virtual FText GetWidgetVisibilityText() const override;
+	virtual FText GetWidgetReadableLocation() const override;
+	virtual FString GetWidgetFile() const override;
+	virtual int32 GetWidgetLineNumber() const override;
+	virtual FName GetWidgetAssetName() const override;
+	virtual FVector2D GetWidgetDesiredSize() const override;
+	virtual FSlateColor GetWidgetForegroundColor() const override;
+	virtual FString GetWidgetAddress() const override;
+	virtual bool GetWidgetEnabled() const override;
+
+private:
+	/** The type string of the widget at the point it was passed to Initialize */
+	UPROPERTY()
+	FText CachedWidgetType;
+
+	/** The visibility string of the widget at the point it was passed to Initialize */
+	UPROPERTY()
+	FText CachedWidgetVisibilityText;
+
+	/** The human readable location (source file for C++ widgets, asset name for UMG widgets) of the widget at the point it was passed to Initialize */
+	UPROPERTY()
+	FText CachedWidgetReadableLocation;
+
+	/** The name of the file that the widget was created from at the point it was passed to Initialize (for C++ widgets) */
+	UPROPERTY()
+	FString CachedWidgetFile;
+
+	/** The line number of the file that the widget was created from at the point it was passed to Initialize (for C++ widgets) */
+	UPROPERTY()
+	int32 CachedWidgetLineNumber;
 
-	/** Should we visualize this node */
-	bool bVisualizeThisNode;
+	/** The name of the asset that the widget was created from at the point it was passed to Initialize (for UMG widgets) */
+	UPROPERTY()
+	FName CachedWidgetAssetName;
 
+	/** The desired size of the widget at the point it was passed to Initialize */
+	UPROPERTY()
+	FVector2D CachedWidgetDesiredSize;
+
+	/** The foreground color of the widget at the point it was passed to Initialize */
+	UPROPERTY()
+	FSlateColor CachedWidgetForegroundColor;
+
+	/** The in-memory address of the widget at the point it was passed to Initialize */
+	UPROPERTY()
+	FString CachedWidgetAddress;
+
+	/** The enabled state of the widget at the point it was passed to Initialize */
+	UPROPERTY()
+	bool CachedWidgetEnabled;
+};
+
+
+class FWidgetReflectorNodeUtils
+{
 public:
+	/**
+	 * Create a single node referencing a live widget.
+	 *
+	 * @param InNodeClass The type of widget reflector node to create
+	 * @param InWidgetGeometry Optional widget and associated geometry which this node should represent
+	 */
+	static ULiveWidgetReflectorNode* NewLiveNode(const FArrangedWidget& InWidgetGeometry = FArrangedWidget(SNullWidget::NullWidget, FGeometry()));
 
 	/**
-	 * FReflectorNodes must be constructed as shared pointers.
+	 * Create nodes for the supplied widget and all their children such that they reference a live widget
+	 * Note that we include both visible and invisible children!
+	 * 
+	 * @param InNodeClass The type of widget reflector node to create
+	 * @param InWidgetGeometry Widget and geometry whose children to capture in the snapshot.
+	 */
+	static ULiveWidgetReflectorNode* NewLiveNodeTreeFrom(const FArrangedWidget& InWidgetGeometry);
+
+	/**
+	 * Create a single node referencing a snapshot of its current state.
 	 *
+	 * @param InNodeClass The type of widget reflector node to create
 	 * @param InWidgetGeometry Optional widget and associated geometry which this node should represent
 	 */
-	static TSharedRef<FReflectorNode> New( const FArrangedWidget& InWidgetGeometry = FArrangedWidget(SNullWidget::NullWidget, FGeometry()) )
-	{
-		return MakeShareable( new FReflectorNode(InWidgetGeometry) );
-	}
+	static USnapshotWidgetReflectorNode* NewSnapshotNode(const FArrangedWidget& InWidgetGeometry = FArrangedWidget(SNullWidget::NullWidget, FGeometry()));
 
 	/**
-	 * Capture all the children of the supplied widget along with their layout results
+	 * Create nodes for the supplied widget and all their children such that they reference a snapshot of their current state
 	 * Note that we include both visible and invisible children!
 	 * 
+	 * @param InNodeClass The type of widget reflector node to create
 	 * @param InWidgetGeometry Widget and geometry whose children to capture in the snapshot.
 	 */
-	static TSharedRef<FReflectorNode> NewTreeFrom( const FArrangedWidget& InWidgetGeometry )
-	{
-		TSharedRef<FReflectorNode> NewNode = MakeShareable( new FReflectorNode(InWidgetGeometry) );
+	static USnapshotWidgetReflectorNode* NewSnapshotNodeTreeFrom(const FArrangedWidget& InWidgetGeometry);
 
-		FArrangedChildren ArrangedChildren(EVisibility::All);
-		NewNode->Widget.Pin()->ArrangeChildren( NewNode->Geometry, ArrangedChildren );
-		
-		for (int32 WidgetIndex=0; WidgetIndex < ArrangedChildren.Num(); ++WidgetIndex)
-		{
-			// Note that we include both visible and invisible children!
-			NewNode->ChildNodes.Add( FReflectorNode::NewTreeFrom( ArrangedChildren[WidgetIndex] )  );
-		}
+private:
+	/**
+	 * Create a single node.
+	 *
+	 * @param InNodeClass The type of widget reflector node to create
+	 * @param InWidgetGeometry Optional widget and associated geometry which this node should represent
+	 */
+	static UWidgetReflectorNodeBase* NewNode(const TSubclassOf<UWidgetReflectorNodeBase>& InNodeClass, const FArrangedWidget& InWidgetGeometry = FArrangedWidget(SNullWidget::NullWidget, FGeometry()));
 
-		return NewNode;
-	}
+	/**
+	 * Create nodes for the supplied widget and all their children
+	 * Note that we include both visible and invisible children!
+	 * 
+	 * @param InNodeClass The type of widget reflector node to create
+	 * @param InWidgetGeometry Widget and geometry whose children to capture in the snapshot.
+	 */
+	static UWidgetReflectorNodeBase* NewNodeTreeFrom(const TSubclassOf<UWidgetReflectorNodeBase>& InNodeClass, const FArrangedWidget& InWidgetGeometry);
 
+public:
 	/**
-	 * Locate all the widgets from a widget path in a list of ReflectorNodes and their children.
+	 * Locate all the widgets from a widget path in a list of nodes and their children.
+	 * @note This only really works for live nodes, as the snapshot nodes may no longer exist, or not even be local to this machine
 	 *
 	 * @param CandidateNodes A list of FReflectorNodes that represent widgets.
 	 * @param WidgetPathToFind We want to find all reflector nodes corresponding to widgets in this path
 	 * @param SearchResult An array that gets results put in it
 	 * @param NodeIndexToFind Index of the widget in the path that we are currently looking for; we are done when we've found all of them
 	 */
-	static void FindWidgetPath( const TArray<TSharedPtr<FReflectorNode>>& CandidateNodes, const FWidgetPath& WidgetPathToFind, TArray< TSharedPtr<FReflectorNode> >& SearchResult, int32 NodeIndexToFind = 0 )
-	{
-		if (NodeIndexToFind < WidgetPathToFind.Widgets.Num())
-		{
-			const FArrangedWidget& WidgetToFind = WidgetPathToFind.Widgets[NodeIndexToFind];
-
-			for (int32 NodeIndex=0; NodeIndex < CandidateNodes.Num(); ++NodeIndex)
-			{
-				if (CandidateNodes[NodeIndex]->Widget.Pin() == WidgetPathToFind.Widgets[NodeIndexToFind].Widget)
-				{
-					SearchResult.Add(CandidateNodes[NodeIndex]);
-					FindWidgetPath(CandidateNodes[NodeIndex]->ChildNodes, WidgetPathToFind, SearchResult, NodeIndexToFind + 1);
-				}
-			}
-		}
-	}
+	static void FindLiveWidgetPath(const TArray<UWidgetReflectorNodeBase*>& CandidateNodes, const FWidgetPath& WidgetPathToFind, TArray<UWidgetReflectorNodeBase*>& SearchResult, int32 NodeIndexToFind = 0);
 
-protected:
+public:
+	/**
+	 * @return The type string for the given widget
+	 */
+	static FText GetWidgetType(const TSharedPtr<SWidget>& InWidget);
+	
+	/**
+	 * @return The current visibility string for the given widget
+	 */
+	static FText GetWidgetVisibilityText(const TSharedPtr<SWidget>& InWidget);
+	
+	/**
+	 * The human readable location for widgets that are defined in C++ is the file and line number
+	 * The human readable location for widgets that are defined in UMG is the asset name
+	 * @return The fully human readable location for the given widget
+	 */
+	static FText GetWidgetReadableLocation(const TSharedPtr<SWidget>& InWidget);
+	
+	/**
+	 * @return The name of the file that this widget was created from (for C++ widgets)
+	 */
+	static FString GetWidgetFile(const TSharedPtr<SWidget>& InWidget);
+	
+	/**
+	 * @return The line number of the file that this widget was created from (for C++ widgets)
+	 */
+	static int32 GetWidgetLineNumber(const TSharedPtr<SWidget>& InWidget);
 
 	/**
-	 * FReflectorNode must be constructor through static methods.
-	 *
-	 * @param InWidgetGeometry Widget and associated geometry that this node will represent
-	 */
-	FReflectorNode( const FArrangedWidget& InWidgetGeometry )
-		: Widget( InWidgetGeometry.Widget )
-		, Geometry( InWidgetGeometry.Geometry )
-		, Tint(FLinearColor(1.0f, 1.0f, 1.0f))
-		, bVisualizeThisNode(true)
-	{ }
- };
+	 * @return The name of the asset that this widget was created from (for UMG widgets)
+	 */
+	static FName GetWidgetAssetName(const TSharedPtr<SWidget>& InWidget);
+
+	/**
+	 * @return The current desired size of the given widget
+	 */
+	static FVector2D GetWidgetDesiredSize(const TSharedPtr<SWidget>& InWidget);
+
+	/**
+	 * @return The in-memory address of the widget, converted to a string
+	 */
+	static FString GetWidgetAddress(const TSharedPtr<SWidget>& InWidget);
+
+	/**
+	 * @return The current foreground color of the given widget
+	 */
+	static FSlateColor GetWidgetForegroundColor(const TSharedPtr<SWidget>& InWidget);
+
+	/**
+	 * @return True if the given widget is currently enabled, false otherwise
+	 */
+	static bool GetWidgetEnabled(const TSharedPtr<SWidget>& InWidget);
+};
diff --git a/Engine/Source/Developer/SlateReflector/Private/Widgets/SWidgetReflector.cpp b/Engine/Source/Developer/SlateReflector/Private/Widgets/SWidgetReflector.cpp
--- a/Engine/Source/Developer/SlateReflector/Private/Widgets/SWidgetReflector.cpp
+++ b/Engine/Source/Developer/SlateReflector/Private/Widgets/SWidgetReflector.cpp
@@ -49,10 +49,10 @@ namespace WidgetReflectorImpl
 /**
  * Widget reflector implementation
  */
-class SWidgetReflector : public ::SWidgetReflector
+class SWidgetReflector : public ::SWidgetReflector, public FGCObject
 {
-	// The reflector uses a tree that observes FReflectorNodes.
-	typedef STreeView<TSharedPtr<FReflectorNode>> SReflectorTree;
+	// The reflector uses a tree that observes UWidgetReflectorNodeBase objects.
+	typedef STreeView<UWidgetReflectorNodeBase*> SReflectorTree;
 
 private:
 
@@ -107,14 +107,14 @@ private:
 	 * @param InReflectorNode The node to generate the tool tip for.
 	 * @return The tool tip widget.
 	 */
-	TSharedRef<SToolTip> GenerateToolTipForReflectorNode( TSharedPtr<FReflectorNode> InReflectorNode );
+	TSharedRef<SToolTip> GenerateToolTipForReflectorNode( UWidgetReflectorNodeBase* InReflectorNode );
 
 	/**
 	 * Mark the provided reflector nodes such that they stand out in the tree and are visible.
 	 *
 	 * @param WidgetPathToObserve The nodes to mark.
 	 */
-	void VisualizeAsTree( const TArray< TSharedPtr<FReflectorNode> >& WidgetPathToVisualize );
+	void VisualizeAsTree( const TArray< UWidgetReflectorNodeBase* >& WidgetPathToVisualize );
 
 	/**
 	 * Draw the widget path to the picked widget as the widgets' outlines.
@@ -135,7 +135,7 @@ private:
 	 * @param LayerId the maximum layer achieved in OutDrawElements so far.
 	 * @return The maximum layer ID we achieved while painting.
 	 */
-	int32 VisualizeSelectedNodesAsRectangles( const TArray<TSharedPtr<FReflectorNode>>& InNodesToDraw, const TSharedRef<SWindow>& VisualizeInWindow, FSlateWindowElementList& OutDrawElements, int32 LayerId );
+	int32 VisualizeSelectedNodesAsRectangles( const TArray<UWidgetReflectorNodeBase*>& InNodesToDraw, const TSharedRef<SWindow>& VisualizeInWindow, FSlateWindowElementList& OutDrawElements, int32 LayerId );
 
 	/** Callback for changing the application scale slider. */
 	void HandleAppScaleSliderChanged( float NewValue )
@@ -191,28 +191,40 @@ private:
 	/** Callback for getting the text of the pick button. */
 	FText HandlePickButtonText() const;
 
+	/** Callback for clicking the snapshot button. */
+	FReply HandleSnapshotButtonClicked();
+
 	/** Callback for generating a row in the reflector tree view. */
-	TSharedRef<ITableRow> HandleReflectorTreeGenerateRow( TSharedPtr<FReflectorNode> InReflectorNode, const TSharedRef<STableViewBase>& OwnerTable );
+	TSharedRef<ITableRow> HandleReflectorTreeGenerateRow( UWidgetReflectorNodeBase* InReflectorNode, const TSharedRef<STableViewBase>& OwnerTable );
 
 	/** Callback for getting the child items of the given reflector tree node. */
-	void HandleReflectorTreeGetChildren( TSharedPtr<FReflectorNode> InWidgetGeometry, TArray<TSharedPtr<FReflectorNode>>& OutChildren );
+	void HandleReflectorTreeGetChildren( UWidgetReflectorNodeBase* InWidgetGeometry, TArray<UWidgetReflectorNodeBase*>& OutChildren );
 
 	/** Callback for when the selection in the reflector tree has changed. */
-	void HandleReflectorTreeSelectionChanged( TSharedPtr< FReflectorNode >, ESelectInfo::Type /*SelectInfo*/ )
+	void HandleReflectorTreeSelectionChanged( UWidgetReflectorNodeBase*, ESelectInfo::Type /*SelectInfo*/ )
 	{
 		SelectedNodes = ReflectorTree->GetSelectedItems();
 	}
 
 	TSharedRef<ITableRow> GenerateEventLogRow( TSharedRef<FLoggedEvent> InReflectorNode, const TSharedRef<STableViewBase>& OwnerTable );
 
+	virtual void AddReferencedObjects( FReferenceCollector& Collector ) override
+	{
+		// STreeView only references items that are currently visible, but we need to prevent the entire tree from being GC'd
+		for (UWidgetReflectorNodeBase* NodeRootItem : ReflectorTreeRoot)
+		{
+			Collector.AddReferencedObject(NodeRootItem);
+		}
+	}
+
 	TArray< TSharedRef<FLoggedEvent> > LoggedEvents;
 	TSharedPtr< SListView< TSharedRef< FLoggedEvent > > > EventListView;
 	TSharedPtr< SHorizontalBox > StatsToolsBox;
 	TSharedPtr<SReflectorTree> ReflectorTree;
 
-	TArray<TSharedPtr<FReflectorNode>> SelectedNodes;
-	TArray<TSharedPtr<FReflectorNode>> ReflectorTreeRoot;
-	TArray<TSharedPtr<FReflectorNode>> PickedPath;
+	TArray<UWidgetReflectorNodeBase*> SelectedNodes;
+	TArray<UWidgetReflectorNodeBase*> ReflectorTreeRoot;
+	TArray<UWidgetReflectorNodeBase*> PickedPath;
 
 	SSplitter::FSlot* WidgetInfoLocation;
 
@@ -356,6 +368,22 @@ void SWidgetReflector::Construct( const FArguments& InArgs )
 									.Text(this, &SWidgetReflector::HandlePickButtonText)
 							]
 					]
+
+					/*
+					+ SHorizontalBox::Slot()
+					.AutoWidth()
+					.Padding(5.0f)
+					[
+						// Button that controls taking a snapshot of the current window(s)
+						SNew(SButton)
+							.OnClicked(this, &SWidgetReflector::HandleSnapshotButtonClicked)
+							[
+								SNew(STextBlock)
+									.Text(LOCTEXT("SnapshotButtonText", "Snapshot"))
+							]
+					]
+					*/
+
 					+SHorizontalBox::Slot()
 					.AutoWidth()
 					.Padding(5.0f)
@@ -611,7 +639,7 @@ void SWidgetReflector::Construct( const FArguments& InArgs )
 
 							+ SHeaderRow::Column("Address")
 							.DefaultLabel( LOCTEXT("Address", "Address") )
-							.FillWidth( 0.10f )
+							.FixedWidth(140.0f)
 						)
 					]
 
@@ -683,7 +711,7 @@ void SWidgetReflector::OnEventProcessed( const FInputEvent& Event, const FReplyB
 
 bool SWidgetReflector::ReflectorNeedsToDrawIn( TSharedRef<SWindow> ThisWindow ) const
 {
-	return ((SelectedNodes.Num() > 0) && (ReflectorTreeRoot.Num() > 0) && (ReflectorTreeRoot[0]->Widget.Pin() == ThisWindow));
+	return ((SelectedNodes.Num() > 0) && (ReflectorTreeRoot.Num() > 0) && (ReflectorTreeRoot[0]->GetLiveWidget() == ThisWindow));
 }
 
 
@@ -693,13 +721,13 @@ void SWidgetReflector::SetWidgetsToVisualize( const FWidgetPath& InWidgetsToVisu
 
 	if (InWidgetsToVisualize.IsValid())
 	{
-		ReflectorTreeRoot.Add(FReflectorNode::NewTreeFrom(InWidgetsToVisualize.Widgets[0]));
+		ReflectorTreeRoot.Add(FWidgetReflectorNodeUtils::NewLiveNodeTreeFrom(InWidgetsToVisualize.Widgets[0]));
 		PickedPath.Empty();
 
-		FReflectorNode::FindWidgetPath( ReflectorTreeRoot, InWidgetsToVisualize, PickedPath );
+		FWidgetReflectorNodeUtils::FindLiveWidgetPath(ReflectorTreeRoot, InWidgetsToVisualize, PickedPath);
 		VisualizeAsTree(PickedPath);
 	}
-	
+
 	ReflectorTree->RequestTreeRefresh();
 }
 
@@ -708,20 +736,21 @@ int32 SWidgetReflector::Visualize( const FWidgetPath& InWidgetsToVisualize, FSla
 {
 	const bool bAttemptingToVisualizeReflector = InWidgetsToVisualize.ContainsWidget(ReflectorTree.ToSharedRef());
 
-	if (!InWidgetsToVisualize.IsValid())
+	if (!InWidgetsToVisualize.IsValid() && SelectedNodes.Num() > 0 && ReflectorTreeRoot.Num() > 0)
 	{
-		TSharedPtr<SWidget> WindowWidget = ReflectorTreeRoot[0]->Widget.Pin();
-		TSharedPtr<SWindow> Window = StaticCastSharedPtr<SWindow>(WindowWidget);
-
-		return VisualizeSelectedNodesAsRectangles(SelectedNodes, Window.ToSharedRef(), OutDrawElements, LayerId);
+		TSharedPtr<SWidget> WindowWidget = ReflectorTreeRoot[0]->GetLiveWidget();
+		if (WindowWidget.IsValid())
+		{
+			TSharedPtr<SWindow> Window = StaticCastSharedPtr<SWindow>(WindowWidget);
+			return VisualizeSelectedNodesAsRectangles(SelectedNodes, Window.ToSharedRef(), OutDrawElements, LayerId);
+		}
 	}
 
 	if (!bAttemptingToVisualizeReflector)
 	{
 		SetWidgetsToVisualize(InWidgetsToVisualize);
-
 		return VisualizePickAsRectangles(InWidgetsToVisualize, OutDrawElements, LayerId);
-	}		
+	}
 
 	return LayerId;
 }
@@ -767,7 +796,7 @@ int32 SWidgetReflector::VisualizeCursorAndKeys(FSlateWindowElementList& OutDrawE
 /* SWidgetReflector implementation
  *****************************************************************************/
 
-TSharedRef<SToolTip> SWidgetReflector::GenerateToolTipForReflectorNode( TSharedPtr<FReflectorNode> InReflectorNode )
+TSharedRef<SToolTip> SWidgetReflector::GenerateToolTipForReflectorNode( UWidgetReflectorNodeBase* InReflectorNode )
 {
 	return SNew(SToolTip)
 		[
@@ -777,18 +806,18 @@ TSharedRef<SToolTip> SWidgetReflector::GenerateToolTipForReflectorNode( TSharedP
 }
 
 
-void SWidgetReflector::VisualizeAsTree( const TArray<TSharedPtr<FReflectorNode>>& WidgetPathToVisualize )
+void SWidgetReflector::VisualizeAsTree( const TArray<UWidgetReflectorNodeBase*>& WidgetPathToVisualize )
 {
 	const FLinearColor TopmostWidgetColor(1.0f, 0.0f, 0.0f);
 	const FLinearColor LeafmostWidgetColor(0.0f, 1.0f, 0.0f);
 
 	for (int32 WidgetIndex = 0; WidgetIndex<WidgetPathToVisualize.Num(); ++WidgetIndex)
 	{
-		TSharedPtr<FReflectorNode> CurWidget = WidgetPathToVisualize[WidgetIndex];
+		UWidgetReflectorNodeBase* CurWidget = WidgetPathToVisualize[WidgetIndex];
 
 		// Tint the item based on depth in picked path
 		const float ColorFactor = static_cast<float>(WidgetIndex)/WidgetPathToVisualize.Num();
-		CurWidget->Tint = FMath::Lerp(TopmostWidgetColor, LeafmostWidgetColor, ColorFactor);
+		CurWidget->SetTint(FMath::Lerp(TopmostWidgetColor, LeafmostWidgetColor, ColorFactor));
 
 		// Make sure the user can see the picked path in the tree.
 		ReflectorTree->SetItemExpansion(CurWidget, true);
@@ -833,11 +862,11 @@ int32 SWidgetReflector::VisualizePickAsRectangles( const FWidgetPath& InWidgetsT
 }
 
 
-int32 SWidgetReflector::VisualizeSelectedNodesAsRectangles( const TArray<TSharedPtr<FReflectorNode>>& InNodesToDraw, const TSharedRef<SWindow>& VisualizeInWindow, FSlateWindowElementList& OutDrawElements, int32 LayerId )
+int32 SWidgetReflector::VisualizeSelectedNodesAsRectangles( const TArray<UWidgetReflectorNodeBase*>& InNodesToDraw, const TSharedRef<SWindow>& VisualizeInWindow, FSlateWindowElementList& OutDrawElements, int32 LayerId )
 {
 	for (int32 NodeIndex = 0; NodeIndex < InNodesToDraw.Num(); ++NodeIndex)
 	{
-		const TSharedPtr<FReflectorNode>& NodeToDraw = InNodesToDraw[NodeIndex];
+		const UWidgetReflectorNodeBase* NodeToDraw = InNodesToDraw[NodeIndex];
 		const FLinearColor Tint(0.0f, 1.0f, 0.0f);
 
 		// The FGeometry we get is from a WidgetPath, so it's rooted in desktop space.
@@ -845,7 +874,7 @@ int32 SWidgetReflector::VisualizeSelectedNodesAsRectangles( const TArray<TShared
 		// and get us back into Window Space.
 		// This is nonstandard so we have to go through some hoops and a specially exposed method 
 		// in FPaintGeometry to allow appending layout transforms.
-		FPaintGeometry WindowSpaceGeometry = NodeToDraw->Geometry.ToPaintGeometry();
+		FPaintGeometry WindowSpaceGeometry = NodeToDraw->GetGeometry().ToPaintGeometry();
 		WindowSpaceGeometry.AppendTransform(TransformCast<FSlateLayoutTransform>(Inverse(VisualizeInWindow->GetPositionInScreen())));
 
 		FSlateDrawElement::MakeBox(
@@ -855,7 +884,7 @@ int32 SWidgetReflector::VisualizeSelectedNodesAsRectangles( const TArray<TShared
 			FCoreStyle::Get().GetBrush(TEXT("Debug.Border")),
 			VisualizeInWindow->GetClippingRectangleInWindow(),
 			ESlateDrawEffect::None,
-			NodeToDraw->Tint
+			NodeToDraw->GetTint()
 		);
 	}
 
@@ -911,7 +940,26 @@ FText SWidgetReflector::HandlePickButtonText() const
 }
 
 
-TSharedRef<ITableRow> SWidgetReflector::HandleReflectorTreeGenerateRow( TSharedPtr<FReflectorNode> InReflectorNode, const TSharedRef<STableViewBase>& OwnerTable )
+FReply SWidgetReflector::HandleSnapshotButtonClicked()
+{
+	// Clean out the current tree and replace it with the snapshotted state of any window(s) that are currently open
+	ReflectorTreeRoot.Empty();
+
+	TArray<TSharedRef<SWindow>> VisibleWindows;
+	FSlateApplication::Get().GetAllVisibleWindowsOrdered(VisibleWindows);
+
+	for (const auto& VisibleWindow : VisibleWindows)
+	{
+		ReflectorTreeRoot.Add(FWidgetReflectorNodeUtils::NewSnapshotNodeTreeFrom(FArrangedWidget(VisibleWindow, VisibleWindow->GetWindowGeometryInScreen())));
+	}
+
+	ReflectorTree->RequestTreeRefresh();
+
+	return FReply::Handled();
+}
+
+
+TSharedRef<ITableRow> SWidgetReflector::HandleReflectorTreeGenerateRow( UWidgetReflectorNodeBase* InReflectorNode, const TSharedRef<STableViewBase>& OwnerTable )
 {
 	return SNew(SReflectorTreeWidgetItem, OwnerTable)
 		.WidgetInfoToVisualize(InReflectorNode)
@@ -921,9 +969,9 @@ TSharedRef<ITableRow> SWidgetReflector::HandleReflectorTreeGenerateRow( TSharedP
 }
 
 
-void SWidgetReflector::HandleReflectorTreeGetChildren(TSharedPtr<FReflectorNode> InWidgetGeometry, TArray<TSharedPtr<FReflectorNode>>& OutChildren)
+void SWidgetReflector::HandleReflectorTreeGetChildren(UWidgetReflectorNodeBase* InWidgetGeometry, TArray<UWidgetReflectorNodeBase*>& OutChildren)
 {
-	OutChildren = InWidgetGeometry->ChildNodes;
+	OutChildren = InWidgetGeometry->GetChildNodes();
 }
 
 
diff --git a/Engine/Source/Developer/SlateReflector/Private/Widgets/SWidgetReflectorToolTipWidget.h b/Engine/Source/Developer/SlateReflector/Private/Widgets/SWidgetReflectorToolTipWidget.h
--- a/Engine/Source/Developer/SlateReflector/Private/Widgets/SWidgetReflectorToolTipWidget.h
+++ b/Engine/Source/Developer/SlateReflector/Private/Widgets/SWidgetReflectorToolTipWidget.h
@@ -15,7 +15,7 @@ public:
 		: _WidgetInfoToVisualize()
 	{ }
 
-		SLATE_ARGUMENT(TSharedPtr<FReflectorNode>, WidgetInfoToVisualize)
+		SLATE_ARGUMENT(UWidgetReflectorNodeBase*, WidgetInfoToVisualize)
 
 	SLATE_END_ARGS()
 
@@ -25,6 +25,7 @@ public:
 	void Construct( const FArguments& InArgs )
 	{
 		this->WidgetInfo = InArgs._WidgetInfoToVisualize;
+		check(WidgetInfo);
 
 		ChildSlot
 		[
@@ -97,43 +98,30 @@ private:
 
 	FText GetWidgetsDesiredSize() const
 	{
-		TSharedPtr<SWidget> TheWidget = WidgetInfo.Get()->Widget.Pin();
-
-		return TheWidget.IsValid()
-			? FText::FromString(TheWidget->GetDesiredSize().ToString())
-			: FText::GetEmpty();
+		return FText::FromString(WidgetInfo->GetWidgetDesiredSize().ToString());
 	}
 
 	FText GetWidgetActualSize() const
 	{
-		return FText::FromString(WidgetInfo.Get()->Geometry.Size.ToString());
+		return FText::FromString(WidgetInfo->GetGeometry().Size.ToString());
 	}
 
 	FText GetSizeInfo() const
 	{
-		TSharedPtr<SWidget> TheWidget = WidgetInfo.Get()->Widget.Pin();
-
-		return TheWidget.IsValid()
-			? FText::FromString(WidgetInfo.Get()->Geometry.ToString())
-			: FText::GetEmpty();
+		return FText::FromString(WidgetInfo->GetGeometry().ToString());
 	}
 
 	FText GetEnabled() const
 	{
 		static const FText TrueText = LOCTEXT("True", "True");
 		static const FText FalseText = LOCTEXT("False", "False");
-
-		TSharedPtr<SWidget> TheWidget = WidgetInfo.Get()->Widget.Pin();
-
-		return TheWidget.IsValid()
-			? (TheWidget->IsEnabled() ? TrueText : FalseText)
-			: FText::GetEmpty();
+		return (WidgetInfo->GetWidgetEnabled()) ? TrueText : FalseText;
 	}
 
 private:
 
 	/** The info about the widget that we are visualizing. */
-	TAttribute<TSharedPtr<FReflectorNode>> WidgetInfo;
+	UWidgetReflectorNodeBase* WidgetInfo;
 };
 
 
diff --git a/Engine/Source/Developer/SlateReflector/Private/Widgets/SWidgetReflectorTreeWidgetItem.cpp b/Engine/Source/Developer/SlateReflector/Private/Widgets/SWidgetReflectorTreeWidgetItem.cpp
--- a/Engine/Source/Developer/SlateReflector/Private/Widgets/SWidgetReflectorTreeWidgetItem.cpp
+++ b/Engine/Source/Developer/SlateReflector/Private/Widgets/SWidgetReflectorTreeWidgetItem.cpp
@@ -1,7 +1,6 @@
 // Copyright 1998-2015 Epic Games, Inc. All Rights Reserved.
 
 #include "SlateReflectorPrivatePCH.h"
-#include "ReflectionMetadata.h"
 #include "SWidgetReflectorTreeWidgetItem.h"
 #include "SHyperlink.h"
 
@@ -58,10 +57,7 @@ TSharedRef<SWidget> SReflectorTreeWidgetItem::GenerateWidgetForColumn(const FNam
 	}
 	else if (ColumnName == "ForegroundColor")
 	{
-		TSharedPtr<SWidget> ThisWidget = WidgetInfo.Get()->Widget.Pin();
-		FSlateColor Foreground = (ThisWidget.IsValid())
-			? ThisWidget->GetForegroundColor()
-			: FSlateColor::UseForeground();
+		const FSlateColor Foreground = WidgetInfo->GetWidgetForegroundColor();
 
 		return SNew(SBorder)
 			// Show unset color as an empty space.
@@ -79,15 +75,18 @@ TSharedRef<SWidget> SReflectorTreeWidgetItem::GenerateWidgetForColumn(const FNam
 	}
 	else if (ColumnName == "Address")
 	{
-		const TSharedPtr<SWidget> TheWidget = WidgetInfo.Get()->Widget.Pin();
-		const FText Address = (TheWidget.IsValid())
-			? FText::FromString(FString::Printf(TEXT("0x%08X"), TheWidget.Get()))
-			: NSLOCTEXT("SWidgetReflector","nullptr","nullptr");
-			
-		return SNew(SHyperlink)
-			.ToolTipText(NSLOCTEXT("SWidgetReflector", "ClickToCopy", "Click to copy address."))
-			.Text(Address)
-			.OnNavigate_Lambda([Address](){ FPlatformMisc::ClipboardCopy(*Address.ToString()); }) ;
+		const FText Address = FText::FromString(WidgetInfo->GetWidgetAddress());
+		
+		return SNew(SBox)
+			.HAlign(HAlign_Left)
+			.VAlign(VAlign_Center)
+			.Padding(FMargin(2.0f, 0.0f))
+			[
+				SNew(SHyperlink)
+				.ToolTipText(NSLOCTEXT("SWidgetReflector", "ClickToCopy", "Click to copy address."))
+				.Text(Address)
+				.OnNavigate_Lambda([Address](){ FPlatformMisc::ClipboardCopy(*Address.ToString()); })
+			];
 	}
 	else
 	{
@@ -96,36 +95,19 @@ TSharedRef<SWidget> SReflectorTreeWidgetItem::GenerateWidgetForColumn(const FNam
 }
 END_SLATE_FUNCTION_BUILD_OPTIMIZATION
 
-FString SReflectorTreeWidgetItem::GetReadableLocation() const
-{
-	TWeakPtr<SWidget> Widget = WidgetInfo.Get()->Widget;
-	if ( Widget.IsValid() )
-	{
-		TSharedPtr<SWidget> SafeWidget = Widget.Pin();
-		TSharedPtr<FReflectionMetaData> MetaData = SafeWidget->GetMetaData<FReflectionMetaData>();
-		if ( MetaData.IsValid() && MetaData->Asset.Get() != nullptr )
-		{
-			return ( MetaData->Asset.IsValid() ? MetaData->Asset->GetName() : FString( TEXT( "Invalid" ) ) ) + TEXT(" [") + MetaData->Name.ToString() + TEXT("]");
-		}
-	}
-
-	return CachedReadableLocation;
-}
-
 void SReflectorTreeWidgetItem::HandleHyperlinkNavigate()
 {
-	TWeakPtr<SWidget> Widget = WidgetInfo.Get()->Widget;
-	if ( Widget.IsValid() )
+	const FName AssetName = WidgetInfo->GetWidgetAssetName();
+	if ( !AssetName.IsNone() )
 	{
-		TSharedPtr<SWidget> SafeWidget = Widget.Pin();
-		TSharedPtr<FReflectionMetaData> MetaData = SafeWidget->GetMetaData<FReflectionMetaData>();
-		if ( MetaData.IsValid() && MetaData->Asset.Get() != nullptr )
+		if ( OnAccessAsset.IsBound() )
 		{
-			if ( OnAccessAsset.IsBound() )
+			UObject* AssetObject = FindObject<UObject>(ANY_PACKAGE, *AssetName.ToString());
+			if ( AssetObject )
 			{
-				OnAccessAsset.Execute(MetaData->Asset.Get());
-				return;
+				OnAccessAsset.Execute(AssetObject);
 			}
+			return;
 		}
 	}
 
diff --git a/Engine/Source/Developer/SlateReflector/Private/Widgets/SWidgetReflectorTreeWidgetItem.h b/Engine/Source/Developer/SlateReflector/Private/Widgets/SWidgetReflectorTreeWidgetItem.h
--- a/Engine/Source/Developer/SlateReflector/Private/Widgets/SWidgetReflectorTreeWidgetItem.h
+++ b/Engine/Source/Developer/SlateReflector/Private/Widgets/SWidgetReflectorTreeWidgetItem.h
@@ -7,7 +7,7 @@
  * Widget that visualizes the contents of a FReflectorNode.
  */
 class SLATEREFLECTOR_API SReflectorTreeWidgetItem
-	: public SMultiColumnTableRow<TSharedPtr<FReflectorNode>>
+	: public SMultiColumnTableRow<UWidgetReflectorNodeBase*>
 {
 public:
 
@@ -17,7 +17,7 @@ public:
 		, _AssetAccessor()
 	{ }
 
-		SLATE_ARGUMENT(TSharedPtr<FReflectorNode>, WidgetInfoToVisualize)
+		SLATE_ARGUMENT(UWidgetReflectorNodeBase*, WidgetInfoToVisualize)
 		SLATE_ARGUMENT(FAccessSourceCode, SourceCodeAccessor)
 		SLATE_ARGUMENT(FAccessAsset, AssetAccessor)
 
@@ -36,18 +36,14 @@ public:
 		this->OnAccessSourceCode = InArgs._SourceCodeAccessor;
 		this->OnAccessAsset = InArgs._AssetAccessor;
 
-		TSharedPtr<SWidget> Widget = WidgetInfo.Get()->Widget.Pin();
-		if ( Widget.IsValid() )
-		{
-			CachedWidgetType = FText::FromString(Widget->GetTypeAsString());
-			CachedWidgetFile = Widget->GetCreatedInLocation().GetPlainNameString();
-			CachedWidgetLineNumber = Widget->GetCreatedInLocation().GetNumber();
-			CachedWidgetVisibility = FText::FromString(Widget->GetVisibility().ToString());
+		check(WidgetInfo);
+		CachedWidgetType = WidgetInfo->GetWidgetType();
+		CachedWidgetVisibility = WidgetInfo->GetWidgetVisibilityText();
+		CachedReadableLocation = WidgetInfo->GetWidgetReadableLocation();
+		CachedWidgetFile = WidgetInfo->GetWidgetFile();
+		CachedWidgetLineNumber = WidgetInfo->GetWidgetLineNumber();
 
-			CachedReadableLocation = Widget->GetReadableLocation();
-		}
-
-		SMultiColumnTableRow< TSharedPtr<FReflectorNode> >::Construct( SMultiColumnTableRow< TSharedPtr<FReflectorNode> >::FArguments().Padding(1), InOwnerTableView );
+		SMultiColumnTableRow< UWidgetReflectorNodeBase* >::Construct( SMultiColumnTableRow< UWidgetReflectorNodeBase* >::FArguments().Padding(1), InOwnerTableView );
 	}
 
 public:
@@ -64,11 +60,14 @@ protected:
 		return CachedWidgetType;
 	}
 	
-	virtual FString GetReadableLocation() const override;
+	virtual FString GetReadableLocation() const override
+	{
+		return CachedReadableLocation.ToString();
+	}
 
 	FText GetReadableLocationAsText() const
 	{
-		return FText::FromString(GetReadableLocation());
+		return CachedReadableLocation;
 	}
 
 	FString GetWidgetFile() const
@@ -89,7 +88,7 @@ protected:
 	/** @return The tint of the reflector node */
 	FSlateColor GetTint() const
 	{
-		return WidgetInfo.Get()->Tint;
+		return WidgetInfo->GetTint();
 	}
 
 	void HandleHyperlinkNavigate();
@@ -97,14 +96,13 @@ protected:
 private:
 
 	/** The info about the widget that we are visualizing. */
-	TAttribute< TSharedPtr<FReflectorNode> > WidgetInfo;
+	UWidgetReflectorNodeBase* WidgetInfo;
 
 	FText CachedWidgetType;
+	FText CachedWidgetVisibility;
+	FText CachedReadableLocation;
 	FString CachedWidgetFile;
 	int32 CachedWidgetLineNumber;
-	FText CachedWidgetVisibility;
-
-	FString CachedReadableLocation;
 
 	FAccessSourceCode OnAccessSourceCode;
 
