Change build condition to match inclusion case for InInvalidationDebuggingEnabled in header

diff --git a/Engine/Source/Developer/SlateReflector/Private/Widgets/SWidgetReflector.cpp b/Engine/Source/Developer/SlateReflector/Private/Widgets/SWidgetReflector.cpp
--- a/Engine/Source/Developer/SlateReflector/Private/Widgets/SWidgetReflector.cpp
+++ b/Engine/Source/Developer/SlateReflector/Private/Widgets/SWidgetReflector.cpp
@@ -433,7 +433,7 @@ void SWidgetReflector::Construct( const FArguments& InArgs )
 							]
 						]
 					]
-#if !UE_BUILD_SHIPPING
+#if !(UE_BUILD_SHIPPING || UE_BUILD_TEST)
 					+ SHorizontalBox::Slot()
 					.AutoWidth()
 					.Padding(5.0f)
