Added the ability to take, save, and load widget snapshots via the widget reflector (Epic Friday)

Major changes:
 - The widget reflector UI now uses tabs. The hierarchy and stats views have been moved to tabs, and new tabs have been added to support picking widgets from a snapshot.
 - The widget reflector node trees are no longer using UObject types. The idea here was to make the serialization easier, but it didn't work and I ended up using JSON instead.
 - When you take or load a widget snapshot, you're now shown the screenshot of the window associated with that snapshot, and are able to pick widgets from it as you would do with a live tree.

API changes:
 - The tab spawners for the widget reflector and atlas visualizers are now registered automatically when the SlateReflector module is loaded. These are now the only way to create these windows.
 - The tabbed layout within the widget reflector has made it impossible to create the widget reflector unless spawned via a tab manager. To this end, the following functions have been removed from ISlateReflectorModule:
  - GetWidgetReflector
  - GetAtlasVisualizer
  - GetTextureAtlasVisualizer
  - GetFontAtlasVisualizer
 - These functions used to be used to create a raw widget reflector or atlas viewer widget.
  - If you were doing this because you were placing it into your own window, you can take advantage of the fact that the widget reflector tab spawner is always available, and use the global tab manager to create your window and place a spawned widget reflector tab inside it - SummonPerfTestSuite in SPerfSuite.cpp provides an example of this.
  - If you were doing this in the handler of your own widget reflector tab spawner, you can either just use the standard widget reflector tab spawner instead (you need to load the SlateReflector module once, and then "WidgetReflector" will be available to spawn via the global tab manager), or you can just use the DisplayWidgetReflector function of ISlateReflectorModule (which will internally do the same thing).

diff --git a/Engine/Source/Developer/SlateReflector/Private/Models/WidgetReflectorNode.cpp b/Engine/Source/Developer/SlateReflector/Private/Models/WidgetReflectorNode.cpp
--- a/Engine/Source/Developer/SlateReflector/Private/Models/WidgetReflectorNode.cpp
+++ b/Engine/Source/Developer/SlateReflector/Private/Models/WidgetReflectorNode.cpp
@@ -8,91 +8,63 @@
 
 /**
  * -----------------------------------------------------------------------------
- * UWidgetReflectorNodeBase
+ * FWidgetReflectorNodeBase
  * -----------------------------------------------------------------------------
  */
-void UWidgetReflectorNodeBase::Initialize(const FArrangedWidget& InWidgetGeometry)
+FWidgetReflectorNodeBase::FWidgetReflectorNodeBase()
+	: Tint(FLinearColor::White)
 {
-	Geometry = InWidgetGeometry.Geometry;
-	Tint = FLinearColor(1.0f, 1.0f, 1.0f);
+	HitTestInfo.IsHitTestVisible = false;
+	HitTestInfo.AreChildrenHitTestVisible = false;
 }
 
-TSharedPtr<SWidget> UWidgetReflectorNodeBase::GetLiveWidget() const
+FWidgetReflectorNodeBase::FWidgetReflectorNodeBase(const FArrangedWidget& InWidgetGeometry)
+	: AccumulatedLayoutTransform(InWidgetGeometry.Geometry.GetAccumulatedLayoutTransform())
+	, AccumulatedRenderTransform(InWidgetGeometry.Geometry.GetAccumulatedRenderTransform())
+	, LocalSize(InWidgetGeometry.Geometry.GetLocalSize())
+	, Tint(FLinearColor::White)
 {
-	return nullptr;
-}
-
-FText UWidgetReflectorNodeBase::GetWidgetType() const
-{
-	return FText::GetEmpty();
-}
-
-FText UWidgetReflectorNodeBase::GetWidgetVisibilityText() const
-{
-	return FText::GetEmpty();
-}
-
-FText UWidgetReflectorNodeBase::GetWidgetReadableLocation() const
-{
-	return FText::GetEmpty();
-}
-
-FString UWidgetReflectorNodeBase::GetWidgetFile() const
-{
-	return FString();
+	const EVisibility WidgetVisibility = InWidgetGeometry.Widget->GetVisibility();
+	HitTestInfo.IsHitTestVisible = WidgetVisibility.IsHitTestVisible();
+	HitTestInfo.AreChildrenHitTestVisible = WidgetVisibility.AreChildrenHitTestVisible();
 }
 
-int32 UWidgetReflectorNodeBase::GetWidgetLineNumber() const
+const FSlateLayoutTransform& FWidgetReflectorNodeBase::GetAccumulatedLayoutTransform() const
 {
-	return 0;
+	return AccumulatedLayoutTransform;
 }
 
-FName UWidgetReflectorNodeBase::GetWidgetAssetName() const
+const FSlateRenderTransform& FWidgetReflectorNodeBase::GetAccumulatedRenderTransform() const
 {
-	return NAME_None;
+	return AccumulatedRenderTransform;
 }
 
-FVector2D UWidgetReflectorNodeBase::GetWidgetDesiredSize() const
+const FVector2D& FWidgetReflectorNodeBase::GetLocalSize() const
 {
-	return FVector2D::ZeroVector;
+	return LocalSize;
 }
 
-FSlateColor UWidgetReflectorNodeBase::GetWidgetForegroundColor() const
+const FWidgetHitTestInfo& FWidgetReflectorNodeBase::GetHitTestInfo() const
 {
-	return FSlateColor::UseForeground();
+	return HitTestInfo;
 }
 
-FString UWidgetReflectorNodeBase::GetWidgetAddress() const
-{
-	return FString::Printf(TEXT("0x%p"), static_cast<const void*>(nullptr));
-}
-
-bool UWidgetReflectorNodeBase::GetWidgetEnabled() const
-{
-	return false;
-}
-
-const FGeometry& UWidgetReflectorNodeBase::GetGeometry() const
-{
-	return Geometry;
-}
-
-const FLinearColor& UWidgetReflectorNodeBase::GetTint() const
+const FLinearColor& FWidgetReflectorNodeBase::GetTint() const
 {
 	return Tint;
 }
 
-void UWidgetReflectorNodeBase::SetTint(const FLinearColor& InTint)
+void FWidgetReflectorNodeBase::SetTint(const FLinearColor& InTint)
 {
 	Tint = InTint;
 }
 
-void UWidgetReflectorNodeBase::AddChildNode(UWidgetReflectorNodeBase* InChildNode)
+void FWidgetReflectorNodeBase::AddChildNode(TSharedRef<FWidgetReflectorNodeBase> InChildNode)
 {
-	ChildNodes.Add(InChildNode);
+	ChildNodes.Add(MoveTemp(InChildNode));
 }
 
-const TArray<UWidgetReflectorNodeBase*>& UWidgetReflectorNodeBase::GetChildNodes() const
+const TArray<TSharedRef<FWidgetReflectorNodeBase>>& FWidgetReflectorNodeBase::GetChildNodes() const
 {
 	return ChildNodes;
 }
@@ -100,67 +72,76 @@ const TArray<UWidgetReflectorNodeBase*>& UWidgetReflectorNodeBase::GetChildNodes
 
 /**
  * -----------------------------------------------------------------------------
- * ULiveWidgetReflectorNode
+ * FLiveWidgetReflectorNode
  * -----------------------------------------------------------------------------
  */
-void ULiveWidgetReflectorNode::Initialize(const FArrangedWidget& InWidgetGeometry)
+TSharedRef<FLiveWidgetReflectorNode> FLiveWidgetReflectorNode::Create(const FArrangedWidget& InWidgetGeometry)
+{
+	return MakeShareable(new FLiveWidgetReflectorNode(InWidgetGeometry));
+}
+
+FLiveWidgetReflectorNode::FLiveWidgetReflectorNode(const FArrangedWidget& InWidgetGeometry)
+	: FWidgetReflectorNodeBase(InWidgetGeometry)
+	, Widget(InWidgetGeometry.Widget)
 {
-	Super::Initialize(InWidgetGeometry);
+}
 
-	Widget = InWidgetGeometry.Widget;
+EWidgetReflectorNodeType FLiveWidgetReflectorNode::GetNodeType() const
+{
+	return EWidgetReflectorNodeType::Live;
 }
 
-TSharedPtr<SWidget> ULiveWidgetReflectorNode::GetLiveWidget() const
+TSharedPtr<SWidget> FLiveWidgetReflectorNode::GetLiveWidget() const
 {
 	return Widget.Pin();
 }
 
-FText ULiveWidgetReflectorNode::GetWidgetType() const
+FText FLiveWidgetReflectorNode::GetWidgetType() const
 {
 	return FWidgetReflectorNodeUtils::GetWidgetType(Widget.Pin());
 }
 
-FText ULiveWidgetReflectorNode::GetWidgetVisibilityText() const
+FText FLiveWidgetReflectorNode::GetWidgetVisibilityText() const
 {
 	return FWidgetReflectorNodeUtils::GetWidgetVisibilityText(Widget.Pin());
 }
 
-FText ULiveWidgetReflectorNode::GetWidgetReadableLocation() const
+FText FLiveWidgetReflectorNode::GetWidgetReadableLocation() const
 {
 	return FWidgetReflectorNodeUtils::GetWidgetReadableLocation(Widget.Pin());
 }
 
-FString ULiveWidgetReflectorNode::GetWidgetFile() const
+FString FLiveWidgetReflectorNode::GetWidgetFile() const
 {
 	return FWidgetReflectorNodeUtils::GetWidgetFile(Widget.Pin());
 }
 
-int32 ULiveWidgetReflectorNode::GetWidgetLineNumber() const
+int32 FLiveWidgetReflectorNode::GetWidgetLineNumber() const
 {
 	return FWidgetReflectorNodeUtils::GetWidgetLineNumber(Widget.Pin());
 }
 
-FName ULiveWidgetReflectorNode::GetWidgetAssetName() const
+FName FLiveWidgetReflectorNode::GetWidgetAssetName() const
 {
 	return FWidgetReflectorNodeUtils::GetWidgetAssetName(Widget.Pin());
 }
 
-FVector2D ULiveWidgetReflectorNode::GetWidgetDesiredSize() const
+FVector2D FLiveWidgetReflectorNode::GetWidgetDesiredSize() const
 {
 	return FWidgetReflectorNodeUtils::GetWidgetDesiredSize(Widget.Pin());
 }
 
-FSlateColor ULiveWidgetReflectorNode::GetWidgetForegroundColor() const
+FSlateColor FLiveWidgetReflectorNode::GetWidgetForegroundColor() const
 {
 	return FWidgetReflectorNodeUtils::GetWidgetForegroundColor(Widget.Pin());
 }
 
-FString ULiveWidgetReflectorNode::GetWidgetAddress() const
+FString FLiveWidgetReflectorNode::GetWidgetAddress() const
 {
 	return FWidgetReflectorNodeUtils::GetWidgetAddress(Widget.Pin());
 }
 
-bool ULiveWidgetReflectorNode::GetWidgetEnabled() const
+bool FLiveWidgetReflectorNode::GetWidgetEnabled() const
 {
 	return FWidgetReflectorNodeUtils::GetWidgetEnabled(Widget.Pin());
 }
@@ -168,111 +149,369 @@ bool ULiveWidgetReflectorNode::GetWidgetEnabled() const
 
 /**
  * -----------------------------------------------------------------------------
- * USnapshotWidgetReflectorNode
+ * FSnapshotWidgetReflectorNode
  * -----------------------------------------------------------------------------
  */
-void USnapshotWidgetReflectorNode::Initialize(const FArrangedWidget& InWidgetGeometry)
+TSharedRef<FSnapshotWidgetReflectorNode> FSnapshotWidgetReflectorNode::Create()
+{
+	return MakeShareable(new FSnapshotWidgetReflectorNode());
+}
+
+TSharedRef<FSnapshotWidgetReflectorNode> FSnapshotWidgetReflectorNode::Create(const FArrangedWidget& InWidgetGeometry)
+{
+	return MakeShareable(new FSnapshotWidgetReflectorNode(InWidgetGeometry));
+}
+
+FSnapshotWidgetReflectorNode::FSnapshotWidgetReflectorNode()
+	: CachedWidgetLineNumber(0)
+	, CachedWidgetEnabled(false)
+{
+}
+
+FSnapshotWidgetReflectorNode::FSnapshotWidgetReflectorNode(const FArrangedWidget& InWidgetGeometry)
+	: FWidgetReflectorNodeBase(InWidgetGeometry)
+	, CachedWidgetType(FWidgetReflectorNodeUtils::GetWidgetType(InWidgetGeometry.Widget))
+	, CachedWidgetVisibilityText(FWidgetReflectorNodeUtils::GetWidgetVisibilityText(InWidgetGeometry.Widget))
+	, CachedWidgetReadableLocation(FWidgetReflectorNodeUtils::GetWidgetReadableLocation(InWidgetGeometry.Widget))
+	, CachedWidgetFile(FWidgetReflectorNodeUtils::GetWidgetFile(InWidgetGeometry.Widget))
+	, CachedWidgetLineNumber(FWidgetReflectorNodeUtils::GetWidgetLineNumber(InWidgetGeometry.Widget))
+	, CachedWidgetAssetName(FWidgetReflectorNodeUtils::GetWidgetAssetName(InWidgetGeometry.Widget))
+	, CachedWidgetDesiredSize(FWidgetReflectorNodeUtils::GetWidgetDesiredSize(InWidgetGeometry.Widget))
+	, CachedWidgetForegroundColor(FWidgetReflectorNodeUtils::GetWidgetForegroundColor(InWidgetGeometry.Widget))
+	, CachedWidgetAddress(FWidgetReflectorNodeUtils::GetWidgetAddress(InWidgetGeometry.Widget))
+	, CachedWidgetEnabled(FWidgetReflectorNodeUtils::GetWidgetEnabled(InWidgetGeometry.Widget))
+{
+}
+
+EWidgetReflectorNodeType FSnapshotWidgetReflectorNode::GetNodeType() const
 {
-	Super::Initialize(InWidgetGeometry);
+	return EWidgetReflectorNodeType::Snapshot;
+}
 
-	CachedWidgetType = FWidgetReflectorNodeUtils::GetWidgetType(InWidgetGeometry.Widget);
-	CachedWidgetVisibilityText = FWidgetReflectorNodeUtils::GetWidgetVisibilityText(InWidgetGeometry.Widget);
-	CachedWidgetReadableLocation = FWidgetReflectorNodeUtils::GetWidgetReadableLocation(InWidgetGeometry.Widget);
-	CachedWidgetFile = FWidgetReflectorNodeUtils::GetWidgetFile(InWidgetGeometry.Widget);
-	CachedWidgetLineNumber = FWidgetReflectorNodeUtils::GetWidgetLineNumber(InWidgetGeometry.Widget);
-	CachedWidgetAssetName = FWidgetReflectorNodeUtils::GetWidgetAssetName(InWidgetGeometry.Widget);
-	CachedWidgetDesiredSize = FWidgetReflectorNodeUtils::GetWidgetDesiredSize(InWidgetGeometry.Widget);
-	CachedWidgetForegroundColor = FWidgetReflectorNodeUtils::GetWidgetForegroundColor(InWidgetGeometry.Widget);
-	CachedWidgetAddress = FWidgetReflectorNodeUtils::GetWidgetAddress(InWidgetGeometry.Widget);
-	CachedWidgetEnabled = FWidgetReflectorNodeUtils::GetWidgetEnabled(InWidgetGeometry.Widget);
+TSharedPtr<SWidget> FSnapshotWidgetReflectorNode::GetLiveWidget() const
+{
+	return nullptr;
 }
 
-FText USnapshotWidgetReflectorNode::GetWidgetType() const
+FText FSnapshotWidgetReflectorNode::GetWidgetType() const
 {
 	return CachedWidgetType;
 }
 
-FText USnapshotWidgetReflectorNode::GetWidgetVisibilityText() const
+FText FSnapshotWidgetReflectorNode::GetWidgetVisibilityText() const
 {
 	return CachedWidgetVisibilityText;
 }
 
-FText USnapshotWidgetReflectorNode::GetWidgetReadableLocation() const
+FText FSnapshotWidgetReflectorNode::GetWidgetReadableLocation() const
 {
 	return CachedWidgetReadableLocation;
 }
 
-FString USnapshotWidgetReflectorNode::GetWidgetFile() const
+FString FSnapshotWidgetReflectorNode::GetWidgetFile() const
 {
 	return CachedWidgetFile;
 }
 
-int32 USnapshotWidgetReflectorNode::GetWidgetLineNumber() const
+int32 FSnapshotWidgetReflectorNode::GetWidgetLineNumber() const
 {
 	return CachedWidgetLineNumber;
 }
 
-FName USnapshotWidgetReflectorNode::GetWidgetAssetName() const
+FName FSnapshotWidgetReflectorNode::GetWidgetAssetName() const
 {
 	return CachedWidgetAssetName;
 }
 
-FVector2D USnapshotWidgetReflectorNode::GetWidgetDesiredSize() const
+FVector2D FSnapshotWidgetReflectorNode::GetWidgetDesiredSize() const
 {
 	return CachedWidgetDesiredSize;
 }
 
-FSlateColor USnapshotWidgetReflectorNode::GetWidgetForegroundColor() const
+FSlateColor FSnapshotWidgetReflectorNode::GetWidgetForegroundColor() const
 {
 	return CachedWidgetForegroundColor;
 }
 
-FString USnapshotWidgetReflectorNode::GetWidgetAddress() const
+FString FSnapshotWidgetReflectorNode::GetWidgetAddress() const
 {
 	return CachedWidgetAddress;
 }
 
-bool USnapshotWidgetReflectorNode::GetWidgetEnabled() const
+bool FSnapshotWidgetReflectorNode::GetWidgetEnabled() const
 {
 	return CachedWidgetEnabled;
 }
 
+TSharedRef<FJsonValue> FSnapshotWidgetReflectorNode::ToJson(const TSharedRef<FSnapshotWidgetReflectorNode>& RootSnapshotNode)
+{
+	struct Internal
+	{
+		static TSharedRef<FJsonValue> CreateVector2DJsonValue(const FVector2D& InVec2D)
+		{
+			TArray<TSharedPtr<FJsonValue>> StructJsonArray;
+			StructJsonArray.Add(MakeShareable(new FJsonValueNumber(InVec2D.X)));
+			StructJsonArray.Add(MakeShareable(new FJsonValueNumber(InVec2D.Y)));
+			return MakeShareable(new FJsonValueArray(StructJsonArray));
+		}
+
+		static TSharedRef<FJsonValue> CreateMatrix2x2JsonValue(const FMatrix2x2& InMatrix)
+		{
+			float m00, m01, m10, m11;
+			InMatrix.GetMatrix(m00, m01, m10, m11);
+
+			TArray<TSharedPtr<FJsonValue>> StructJsonArray;
+			StructJsonArray.Add(MakeShareable(new FJsonValueNumber(m00)));
+			StructJsonArray.Add(MakeShareable(new FJsonValueNumber(m01)));
+			StructJsonArray.Add(MakeShareable(new FJsonValueNumber(m10)));
+			StructJsonArray.Add(MakeShareable(new FJsonValueNumber(m11)));
+			return MakeShareable(new FJsonValueArray(StructJsonArray));
+		}
+
+		static TSharedRef<FJsonValue> CreateSlateLayoutTransformJsonValue(const FSlateLayoutTransform& InLayoutTransform)
+		{
+			TSharedRef<FJsonObject> StructJsonObject = MakeShareable(new FJsonObject());
+			StructJsonObject->SetNumberField(TEXT("Scale"), InLayoutTransform.GetScale());
+			StructJsonObject->SetField(TEXT("Translation"), CreateVector2DJsonValue(InLayoutTransform.GetTranslation()));
+			return MakeShareable(new FJsonValueObject(StructJsonObject));
+		}
+
+		static TSharedRef<FJsonValue> CreateSlateRenderTransformJsonValue(const FSlateRenderTransform& InRenderTransform)
+		{
+			TSharedRef<FJsonObject> StructJsonObject = MakeShareable(new FJsonObject());
+			StructJsonObject->SetField(TEXT("Matrix"), CreateMatrix2x2JsonValue(InRenderTransform.GetMatrix()));
+			StructJsonObject->SetField(TEXT("Translation"), CreateVector2DJsonValue(InRenderTransform.GetTranslation()));
+			return MakeShareable(new FJsonValueObject(StructJsonObject));
+		}
+
+		static TSharedRef<FJsonValue> CreateLinearColorJsonValue(const FLinearColor& InColor)
+		{
+			TArray<TSharedPtr<FJsonValue>> StructJsonArray;
+			StructJsonArray.Add(MakeShareable(new FJsonValueNumber(InColor.R)));
+			StructJsonArray.Add(MakeShareable(new FJsonValueNumber(InColor.G)));
+			StructJsonArray.Add(MakeShareable(new FJsonValueNumber(InColor.B)));
+			StructJsonArray.Add(MakeShareable(new FJsonValueNumber(InColor.A)));
+			return MakeShareable(new FJsonValueArray(StructJsonArray));
+		}
+
+		static TSharedRef<FJsonValue> CreateSlateColorJsonValue(const FSlateColor& InColor)
+		{
+			const bool bIsColorSpecified = InColor.IsColorSpecified();
+			const FLinearColor ColorToUse = (bIsColorSpecified) ? InColor.GetSpecifiedColor() : FLinearColor::White;
+
+			TSharedRef<FJsonObject> StructJsonObject = MakeShareable(new FJsonObject());
+			StructJsonObject->SetBoolField(TEXT("IsColorSpecified"), bIsColorSpecified);
+			StructJsonObject->SetField(TEXT("Color"), CreateLinearColorJsonValue(ColorToUse));
+			return MakeShareable(new FJsonValueObject(StructJsonObject));
+		}
+
+		static TSharedRef<FJsonValue> CreateWidgetHitTestInfoJsonValue(const FWidgetHitTestInfo& InHitTestInfo)
+		{
+			TSharedRef<FJsonObject> StructJsonObject = MakeShareable(new FJsonObject());
+			StructJsonObject->SetBoolField(TEXT("IsHitTestVisible"), InHitTestInfo.IsHitTestVisible);
+			StructJsonObject->SetBoolField(TEXT("AreChildrenHitTestVisible"), InHitTestInfo.AreChildrenHitTestVisible);
+			return MakeShareable(new FJsonValueObject(StructJsonObject));
+		}
+	};
+
+	TSharedRef<FJsonObject> RootJsonObject = MakeShareable(new FJsonObject());
+
+	RootJsonObject->SetField(TEXT("AccumulatedLayoutTransform"), Internal::CreateSlateLayoutTransformJsonValue(RootSnapshotNode->AccumulatedLayoutTransform));
+	RootJsonObject->SetField(TEXT("AccumulatedRenderTransform"), Internal::CreateSlateRenderTransformJsonValue(RootSnapshotNode->AccumulatedRenderTransform));
+	RootJsonObject->SetField(TEXT("LocalSize"), Internal::CreateVector2DJsonValue(RootSnapshotNode->LocalSize));
+	RootJsonObject->SetField(TEXT("HitTestInfo"), Internal::CreateWidgetHitTestInfoJsonValue(RootSnapshotNode->HitTestInfo));
+	RootJsonObject->SetStringField(TEXT("WidgetType"), RootSnapshotNode->CachedWidgetType.ToString());
+	RootJsonObject->SetStringField(TEXT("WidgetVisibilityText"), RootSnapshotNode->CachedWidgetVisibilityText.ToString());
+	RootJsonObject->SetStringField(TEXT("WidgetReadableLocation"), RootSnapshotNode->CachedWidgetReadableLocation.ToString());
+	RootJsonObject->SetStringField(TEXT("WidgetFile"), RootSnapshotNode->CachedWidgetFile);
+	RootJsonObject->SetNumberField(TEXT("WidgetLineNumber"), RootSnapshotNode->CachedWidgetLineNumber);
+	RootJsonObject->SetStringField(TEXT("WidgetAssetName"), RootSnapshotNode->CachedWidgetAssetName.ToString());
+	RootJsonObject->SetField(TEXT("WidgetDesiredSize"), Internal::CreateVector2DJsonValue(RootSnapshotNode->CachedWidgetDesiredSize));
+	RootJsonObject->SetField(TEXT("WidgetForegroundColor"), Internal::CreateSlateColorJsonValue(RootSnapshotNode->CachedWidgetForegroundColor));
+	RootJsonObject->SetStringField(TEXT("WidgetAddress"), RootSnapshotNode->CachedWidgetAddress);
+	RootJsonObject->SetBoolField(TEXT("WidgetEnabled"), RootSnapshotNode->CachedWidgetEnabled);
+
+	TArray<TSharedPtr<FJsonValue>> ChildNodesJsonArray;
+	for (const auto& ChildReflectorNode : RootSnapshotNode->ChildNodes)
+	{
+		check(ChildReflectorNode->GetNodeType() == EWidgetReflectorNodeType::Snapshot);
+		ChildNodesJsonArray.Add(FSnapshotWidgetReflectorNode::ToJson(StaticCastSharedRef<FSnapshotWidgetReflectorNode>(ChildReflectorNode)));
+	}
+	RootJsonObject->SetArrayField(TEXT("ChildNodes"), ChildNodesJsonArray);
+
+	return MakeShareable(new FJsonValueObject(RootJsonObject));
+}
+
+TSharedRef<FSnapshotWidgetReflectorNode> FSnapshotWidgetReflectorNode::FromJson(const TSharedRef<FJsonValue>& RootJsonValue)
+{
+	struct Internal
+	{
+		static FVector2D ParseVector2DJsonValue(const TSharedPtr<FJsonValue>& InJsonValue)
+		{
+			check(InJsonValue.IsValid());
+			const TArray<TSharedPtr<FJsonValue>>& StructJsonArray = InJsonValue->AsArray();
+			check(StructJsonArray.Num() == 2);
+
+			return FVector2D(
+				StructJsonArray[0]->AsNumber(), 
+				StructJsonArray[1]->AsNumber()
+				);
+		}
+
+		static FMatrix2x2 ParseMatrix2x2JsonValue(const TSharedPtr<FJsonValue>& InJsonValue)
+		{
+			check(InJsonValue.IsValid());
+			const TArray<TSharedPtr<FJsonValue>>& StructJsonArray = InJsonValue->AsArray();
+			check(StructJsonArray.Num() == 4);
+
+			return FMatrix2x2(
+				StructJsonArray[0]->AsNumber(), 
+				StructJsonArray[1]->AsNumber(), 
+				StructJsonArray[2]->AsNumber(), 
+				StructJsonArray[3]->AsNumber()
+				);
+		}
+
+		static FSlateLayoutTransform ParseSlateLayoutTransformJsonValue(const TSharedPtr<FJsonValue>& InJsonValue)
+		{
+			check(InJsonValue.IsValid());
+			const TSharedPtr<FJsonObject>& StructJsonObject = InJsonValue->AsObject();
+			check(StructJsonObject.IsValid());
+
+			return FSlateLayoutTransform(
+				StructJsonObject->GetNumberField(TEXT("Scale")),
+				ParseVector2DJsonValue(StructJsonObject->GetField<EJson::None>(TEXT("Translation")))
+				);
+		}
+
+		static FSlateRenderTransform ParseSlateRenderTransformJsonValue(const TSharedPtr<FJsonValue>& InJsonValue)
+		{
+			check(InJsonValue.IsValid());
+			const TSharedPtr<FJsonObject>& StructJsonObject = InJsonValue->AsObject();
+			check(StructJsonObject.IsValid());
+
+			return FSlateRenderTransform(
+				ParseMatrix2x2JsonValue(StructJsonObject->GetField<EJson::None>(TEXT("Matrix"))),
+				ParseVector2DJsonValue(StructJsonObject->GetField<EJson::None>(TEXT("Translation")))
+				);
+		}
+
+		static FLinearColor ParseLinearColorJsonValue(const TSharedPtr<FJsonValue>& InJsonValue)
+		{
+			check(InJsonValue.IsValid());
+			const TArray<TSharedPtr<FJsonValue>>& StructJsonArray = InJsonValue->AsArray();
+			check(StructJsonArray.Num() == 4);
+
+			return FLinearColor(
+				StructJsonArray[0]->AsNumber(), 
+				StructJsonArray[1]->AsNumber(), 
+				StructJsonArray[2]->AsNumber(), 
+				StructJsonArray[3]->AsNumber()
+				);
+		}
+
+		static FSlateColor ParseSlateColorJsonValue(const TSharedPtr<FJsonValue>& InJsonValue)
+		{
+			check(InJsonValue.IsValid());
+			const TSharedPtr<FJsonObject>& StructJsonObject = InJsonValue->AsObject();
+			check(StructJsonObject.IsValid());
+
+			const bool bIsColorSpecified = StructJsonObject->GetBoolField(TEXT("IsColorSpecified"));
+			if (bIsColorSpecified)
+			{
+				return FSlateColor(ParseLinearColorJsonValue(StructJsonObject->GetField<EJson::None>(TEXT("Color"))));
+			}
+			else
+			{
+				return FSlateColor::UseForeground();
+			}
+		}
+
+		static FWidgetHitTestInfo ParseWidgetHitTestInfoJsonValue(const TSharedPtr<FJsonValue>& InJsonValue)
+		{
+			check(InJsonValue.IsValid());
+			const TSharedPtr<FJsonObject>& StructJsonObject = InJsonValue->AsObject();
+			check(StructJsonObject.IsValid());
+
+			FWidgetHitTestInfo HitTestInfo;
+			HitTestInfo.IsHitTestVisible = StructJsonObject->GetBoolField(TEXT("IsHitTestVisible"));
+			HitTestInfo.AreChildrenHitTestVisible = StructJsonObject->GetBoolField(TEXT("AreChildrenHitTestVisible"));
+			return HitTestInfo;
+		}
+	};
+
+	const TSharedPtr<FJsonObject>& RootJsonObject = RootJsonValue->AsObject();
+	check(RootJsonObject.IsValid());
+
+	auto RootSnapshotNode = FSnapshotWidgetReflectorNode::Create();
+
+	RootSnapshotNode->AccumulatedLayoutTransform = Internal::ParseSlateLayoutTransformJsonValue(RootJsonObject->GetField<EJson::None>(TEXT("AccumulatedLayoutTransform")));
+	RootSnapshotNode->AccumulatedRenderTransform = Internal::ParseSlateRenderTransformJsonValue(RootJsonObject->GetField<EJson::None>(TEXT("AccumulatedRenderTransform")));
+	RootSnapshotNode->LocalSize = Internal::ParseVector2DJsonValue(RootJsonObject->GetField<EJson::None>(TEXT("LocalSize")));
+	RootSnapshotNode->HitTestInfo = Internal::ParseWidgetHitTestInfoJsonValue(RootJsonObject->GetField<EJson::None>(TEXT("HitTestInfo")));
+	RootSnapshotNode->CachedWidgetType = FText::FromString(RootJsonObject->GetStringField(TEXT("WidgetType")));
+	RootSnapshotNode->CachedWidgetVisibilityText = FText::FromString(RootJsonObject->GetStringField(TEXT("WidgetVisibilityText")));
+	RootSnapshotNode->CachedWidgetReadableLocation = FText::FromString(RootJsonObject->GetStringField(TEXT("WidgetReadableLocation")));
+	RootSnapshotNode->CachedWidgetFile = RootJsonObject->GetStringField(TEXT("WidgetFile"));
+	RootSnapshotNode->CachedWidgetLineNumber = RootJsonObject->GetIntegerField(TEXT("WidgetLineNumber"));
+	RootSnapshotNode->CachedWidgetAssetName = FName(*RootJsonObject->GetStringField(TEXT("WidgetAssetName")));
+	RootSnapshotNode->CachedWidgetDesiredSize = Internal::ParseVector2DJsonValue(RootJsonObject->GetField<EJson::None>(TEXT("WidgetDesiredSize")));
+	RootSnapshotNode->CachedWidgetForegroundColor = Internal::ParseSlateColorJsonValue(RootJsonObject->GetField<EJson::None>(TEXT("WidgetForegroundColor")));
+	RootSnapshotNode->CachedWidgetAddress = RootJsonObject->GetStringField(TEXT("WidgetAddress"));
+	RootSnapshotNode->CachedWidgetEnabled = RootJsonObject->GetBoolField(TEXT("WidgetEnabled"));
+
+	const TArray<TSharedPtr<FJsonValue>>& ChildNodesJsonArray = RootJsonObject->GetArrayField(TEXT("ChildNodes"));
+	for (const TSharedPtr<FJsonValue>& ChildNodeJsonValue : ChildNodesJsonArray)
+	{
+		RootSnapshotNode->AddChildNode(FSnapshotWidgetReflectorNode::FromJson(ChildNodeJsonValue.ToSharedRef()));
+	}
+
+	return RootSnapshotNode;
+}
 
 /**
  * -----------------------------------------------------------------------------
  * FWidgetReflectorNodeUtils
  * -----------------------------------------------------------------------------
  */
-ULiveWidgetReflectorNode* FWidgetReflectorNodeUtils::NewLiveNode(const FArrangedWidget& InWidgetGeometry)
+TSharedRef<FLiveWidgetReflectorNode> FWidgetReflectorNodeUtils::NewLiveNode(const FArrangedWidget& InWidgetGeometry)
 {
-	return Cast<ULiveWidgetReflectorNode>(NewNode(ULiveWidgetReflectorNode::StaticClass(), InWidgetGeometry));
+	return StaticCastSharedRef<FLiveWidgetReflectorNode>(NewNode(EWidgetReflectorNodeType::Live, InWidgetGeometry));
 }
 
-ULiveWidgetReflectorNode* FWidgetReflectorNodeUtils::NewLiveNodeTreeFrom(const FArrangedWidget& InWidgetGeometry)
+TSharedRef<FLiveWidgetReflectorNode> FWidgetReflectorNodeUtils::NewLiveNodeTreeFrom(const FArrangedWidget& InWidgetGeometry)
 {
-	return Cast<ULiveWidgetReflectorNode>(NewNodeTreeFrom(ULiveWidgetReflectorNode::StaticClass(), InWidgetGeometry));
+	return StaticCastSharedRef<FLiveWidgetReflectorNode>(NewNodeTreeFrom(EWidgetReflectorNodeType::Live, InWidgetGeometry));
 }
 
-USnapshotWidgetReflectorNode* FWidgetReflectorNodeUtils::NewSnapshotNode(const FArrangedWidget& InWidgetGeometry)
+TSharedRef<FSnapshotWidgetReflectorNode> FWidgetReflectorNodeUtils::NewSnapshotNode(const FArrangedWidget& InWidgetGeometry)
 {
-	return Cast<USnapshotWidgetReflectorNode>(NewNode(USnapshotWidgetReflectorNode::StaticClass(), InWidgetGeometry));
+	return StaticCastSharedRef<FSnapshotWidgetReflectorNode>(NewNode(EWidgetReflectorNodeType::Snapshot, InWidgetGeometry));
 }
 
-USnapshotWidgetReflectorNode* FWidgetReflectorNodeUtils::NewSnapshotNodeTreeFrom(const FArrangedWidget& InWidgetGeometry)
+TSharedRef<FSnapshotWidgetReflectorNode> FWidgetReflectorNodeUtils::NewSnapshotNodeTreeFrom(const FArrangedWidget& InWidgetGeometry)
 {
-	return Cast<USnapshotWidgetReflectorNode>(NewNodeTreeFrom(USnapshotWidgetReflectorNode::StaticClass(), InWidgetGeometry));
+	return StaticCastSharedRef<FSnapshotWidgetReflectorNode>(NewNodeTreeFrom(EWidgetReflectorNodeType::Snapshot, InWidgetGeometry));
 }
 
-UWidgetReflectorNodeBase* FWidgetReflectorNodeUtils::NewNode(const TSubclassOf<UWidgetReflectorNodeBase>& InNodeClass, const FArrangedWidget& InWidgetGeometry)
+TSharedRef<FWidgetReflectorNodeBase> FWidgetReflectorNodeUtils::NewNode(const EWidgetReflectorNodeType InNodeType, const FArrangedWidget& InWidgetGeometry)
 {
-	UWidgetReflectorNodeBase* NewNodeInstance = NewObject<UWidgetReflectorNodeBase>(GetTransientPackage(), *InNodeClass);
-	NewNodeInstance->Initialize(InWidgetGeometry);
-	return NewNodeInstance;
+	switch (InNodeType)
+	{
+	case EWidgetReflectorNodeType::Live:
+		return FLiveWidgetReflectorNode::Create(InWidgetGeometry);
+	case EWidgetReflectorNodeType::Snapshot:
+		return FSnapshotWidgetReflectorNode::Create(InWidgetGeometry);
+	default:
+		// Should never reach this point, but we have to return something!
+		check(false);
+		return FLiveWidgetReflectorNode::Create(InWidgetGeometry);
+	}
 }
 
-UWidgetReflectorNodeBase* FWidgetReflectorNodeUtils::NewNodeTreeFrom(const TSubclassOf<UWidgetReflectorNodeBase>& InNodeClass, const FArrangedWidget& InWidgetGeometry)
+TSharedRef<FWidgetReflectorNodeBase> FWidgetReflectorNodeUtils::NewNodeTreeFrom(const EWidgetReflectorNodeType InNodeType, const FArrangedWidget& InWidgetGeometry)
 {
-	UWidgetReflectorNodeBase* NewNodeInstance = NewNode(InNodeClass, InWidgetGeometry);
+	TSharedRef<FWidgetReflectorNodeBase> NewNodeInstance = NewNode(InNodeType, InWidgetGeometry);
 
 	FArrangedChildren ArrangedChildren(EVisibility::All);
 	InWidgetGeometry.Widget->ArrangeChildren(InWidgetGeometry.Geometry, ArrangedChildren);
@@ -280,13 +519,13 @@ UWidgetReflectorNodeBase* FWidgetReflectorNodeUtils::NewNodeTreeFrom(const TSubc
 	for (int32 WidgetIndex = 0; WidgetIndex < ArrangedChildren.Num(); ++WidgetIndex)
 	{
 		// Note that we include both visible and invisible children!
-		NewNodeInstance->AddChildNode(NewNodeTreeFrom(InNodeClass, ArrangedChildren[WidgetIndex]));
+		NewNodeInstance->AddChildNode(NewNodeTreeFrom(InNodeType, ArrangedChildren[WidgetIndex]));
 	}
 
 	return NewNodeInstance;
 }
 
-void FWidgetReflectorNodeUtils::FindLiveWidgetPath(const TArray<UWidgetReflectorNodeBase*>& CandidateNodes, const FWidgetPath& WidgetPathToFind, TArray<UWidgetReflectorNodeBase*>& SearchResult, int32 NodeIndexToFind)
+void FWidgetReflectorNodeUtils::FindLiveWidgetPath(const TArray<TSharedRef<FWidgetReflectorNodeBase>>& CandidateNodes, const FWidgetPath& WidgetPathToFind, TArray<TSharedRef<FWidgetReflectorNodeBase>>& SearchResult, int32 NodeIndexToFind)
 {
 	if (NodeIndexToFind < WidgetPathToFind.Widgets.Num())
 	{
diff --git a/Engine/Source/Developer/SlateReflector/Private/Models/WidgetReflectorNode.h b/Engine/Source/Developer/SlateReflector/Private/Models/WidgetReflectorNode.h
--- a/Engine/Source/Developer/SlateReflector/Private/Models/WidgetReflectorNode.h
+++ b/Engine/Source/Developer/SlateReflector/Private/Models/WidgetReflectorNode.h
@@ -2,84 +2,117 @@
 
 #pragma once
 
-#include "WidgetReflectorNode.generated.h"
+/**
+ * Enum used as crude RTTI for the widget reflector node types
+ */
+enum class EWidgetReflectorNodeType : uint8
+{
+	Live,
+	Snapshot,
+};
+
+/**
+ * Cached information about whether a widget can be hit-tested
+ */
+struct FWidgetHitTestInfo
+{
+	bool IsHitTestVisible;
+	bool AreChildrenHitTestVisible;
+};
 
 /** 
  * A widget reflector node that contains the interface and basic data required by both live and snapshot nodes 
  */
-UCLASS()
-class UWidgetReflectorNodeBase : public UObject
+class FWidgetReflectorNodeBase
 {
-	GENERATED_BODY()
-
 public:
 	/**
-	 * Initialize this node from the given widget geometry, caching out any data that may be required for future visualization in the widget reflector
+	 * Destructor
 	 */
-	virtual void Initialize(const FArrangedWidget& InWidgetGeometry);
+	virtual ~FWidgetReflectorNodeBase() {}
+
+	/**
+	 * @return Get the enum entry corresponding to this type of widget reflector node (used as crude RTTI)
+	 */
+	virtual EWidgetReflectorNodeType GetNodeType() const = 0;
 
 	/**
 	 * @note This function only works for ULiveWidgetReflectorNode instances
 	 * @return The live widget that this node is referencing
 	 */
-	virtual TSharedPtr<SWidget> GetLiveWidget() const;
+	virtual TSharedPtr<SWidget> GetLiveWidget() const = 0;
 
 	/**
 	 * @return The type string for the widget we were initialized from
 	 */
-	virtual FText GetWidgetType() const;
+	virtual FText GetWidgetType() const = 0;
 	
 	/**
 	 * @return The visibility string for the widget we were initialized from
 	 */
-	virtual FText GetWidgetVisibilityText() const;
+	virtual FText GetWidgetVisibilityText() const = 0;
 
 	/**
 	 * The human readable location for widgets that are defined in C++ is the file and line number
 	 * The human readable location for widgets that are defined in UMG is the asset name
 	 * @return The fully human readable location for the widget we were initialized from
 	 */
-	virtual FText GetWidgetReadableLocation() const;
+	virtual FText GetWidgetReadableLocation() const = 0;
 
 	/**
 	 * @return The name of the file that the widget we were initialized from was created from (for C++ widgets)
 	 */
-	virtual FString GetWidgetFile() const;
+	virtual FString GetWidgetFile() const = 0;
 
 	/**
 	 * @return The line number of the file that the widget we were initialized from was created from (for C++ widgets)
 	 */
-	virtual int32 GetWidgetLineNumber() const;
+	virtual int32 GetWidgetLineNumber() const = 0;
 
 	/**
 	 * @return The name of the asset that the widget we were initialized from was created from (for UMG widgets)
 	 */
-	virtual FName GetWidgetAssetName() const;
+	virtual FName GetWidgetAssetName() const = 0;
 
 	/**
 	 * @return The desired size of the widget we were initialized from
 	 */
-	virtual FVector2D GetWidgetDesiredSize() const;
+	virtual FVector2D GetWidgetDesiredSize() const = 0;
 
 	/**
 	 * @return The foreground color of the widget we were initialized from
 	 */
-	virtual FSlateColor GetWidgetForegroundColor() const;
+	virtual FSlateColor GetWidgetForegroundColor() const = 0;
 
 	/**
 	 * @return The in-memory address of the widget we were initialized from
 	 */
-	virtual FString GetWidgetAddress() const;
+	virtual FString GetWidgetAddress() const = 0;
 
 	/**
 	 * @return True if the the widget we were initialized from is enabled, false otherwise
 	 */
-	virtual bool GetWidgetEnabled() const;
+	virtual bool GetWidgetEnabled() const = 0;
+
+	/**
+	 * @return The accumulated layout transform of the widget we were initialized from
+	 */
+	const FSlateLayoutTransform& GetAccumulatedLayoutTransform() const;
 
 	/**
-	 * @return The geometry of the widget we were initialized from
+	 * @return The accumulated render transform of the widget we were initialized from
 	 */
-	const FGeometry& GetGeometry() const;
+	const FSlateRenderTransform& GetAccumulatedRenderTransform() const;
+
+	/**
+	 * @return The local size of the widget we were initialized from
+	 */
+	const FVector2D& GetLocalSize() const;
+
+	/**
+	 * @return The basic hit-test of the widget we were initialized from
+	 */
+	const FWidgetHitTestInfo& GetHitTestInfo() const;
 
 	/**
 	 * @return The tint that is applied to text in order to provide visual hints
@@ -92,40 +125,64 @@ public:
 	void SetTint(const FLinearColor& InTint);
 
 	/**
-	 * Add the given node to our list of children for this widget
+	 * Add the given node to our list of children for this widget (this node will keep a strong reference to the instance)
 	 */
-	void AddChildNode(UWidgetReflectorNodeBase* InChildNode);
+	void AddChildNode(TSharedRef<FWidgetReflectorNodeBase> InChildNode);
 
 	/**
 	 * @return The node entries for the widget's children
 	 */
-	const TArray<UWidgetReflectorNodeBase*>& GetChildNodes() const;
+	const TArray<TSharedRef<FWidgetReflectorNodeBase>>& GetChildNodes() const;
 
 protected:
-	/** The geometry of the widget */
-	UPROPERTY()
-	FGeometry Geometry;
+	/**
+	 * Default constructor
+	 */
+	FWidgetReflectorNodeBase();
+
+	/**
+	 * Construct this node from the given widget geometry, caching out any data that may be required for future visualization in the widget reflector
+	 */
+	explicit FWidgetReflectorNodeBase(const FArrangedWidget& InWidgetGeometry);
+
+protected:
+	/** The accumulated layout transform of the widget */
+	FSlateLayoutTransform AccumulatedLayoutTransform;
+
+	/** The accumulated render transform of the widget */
+	FSlateRenderTransform AccumulatedRenderTransform;
+
+	/** The local size of the widget */
+	FVector2D LocalSize;
+
+	/** The hit-test information for the widget */
+	FWidgetHitTestInfo HitTestInfo;
 
 	/** Node entries for the widget's children */
-	UPROPERTY()
-	TArray<UWidgetReflectorNodeBase*> ChildNodes;
+	TArray<TSharedRef<FWidgetReflectorNodeBase>> ChildNodes;
 
-	/** A tint that is applied to text in order to provide visual hints */
-	UPROPERTY()
+	/** A tint that is applied to text in order to provide visual hints (Transient) */
 	FLinearColor Tint;
 };
 
 /** 
  * A widget reflector node that holds on to the widget it references so that certain properties can be updated live 
  */
-UCLASS()
-class ULiveWidgetReflectorNode : public UWidgetReflectorNodeBase
+class FLiveWidgetReflectorNode : public FWidgetReflectorNodeBase
 {
-	GENERATED_BODY()
-
 public:
-	// UWidgetReflectorNodeBase interface
-	virtual void Initialize(const FArrangedWidget& InWidgetGeometry) override;
+	/**
+	 * Destructor
+	 */
+	virtual ~FLiveWidgetReflectorNode() {}
+
+	/**
+	 * Create a live node instance from the given widget geometry, caching out any data that may be required for future visualization in the widget reflector
+	 */
+	static TSharedRef<FLiveWidgetReflectorNode> Create(const FArrangedWidget& InWidgetGeometry);
+
+	// FWidgetReflectorNodeBase interface
+	virtual EWidgetReflectorNodeType GetNodeType() const override;
 	virtual TSharedPtr<SWidget> GetLiveWidget() const override;
 	virtual FText GetWidgetType() const override;
 	virtual FText GetWidgetVisibilityText() const override;
@@ -138,6 +195,12 @@ public:
 	virtual FString GetWidgetAddress() const override;
 	virtual bool GetWidgetEnabled() const override;
 
+private:
+	/**
+	 * Construct this node from the given widget geometry, caching out any data that may be required for future visualization in the widget reflector
+	 */
+	explicit FLiveWidgetReflectorNode(const FArrangedWidget& InWidgetGeometry);
+
 private:
 	/** The widget this node is watching */
 	TWeakPtr<SWidget> Widget;
@@ -146,14 +209,27 @@ private:
 /** 
  * A widget reflector node that holds the widget information from a snapshot at a given point in time 
  */
-UCLASS()
-class USnapshotWidgetReflectorNode : public UWidgetReflectorNodeBase
+class FSnapshotWidgetReflectorNode : public FWidgetReflectorNodeBase
 {
-	GENERATED_BODY()
-
 public:
-	// UWidgetReflectorNodeBase interface
-	virtual void Initialize(const FArrangedWidget& InWidgetGeometry) override;
+	/**
+	 * Destructor
+	 */
+	virtual ~FSnapshotWidgetReflectorNode() {}
+
+	/**
+	 * Create a default snapshot node instance
+	 */
+	static TSharedRef<FSnapshotWidgetReflectorNode> Create();
+
+	/**
+	 * Create a snapshot node instance node from the given widget geometry, caching out any data that may be required for future visualization in the widget reflector
+	 */
+	static TSharedRef<FSnapshotWidgetReflectorNode> Create(const FArrangedWidget& InWidgetGeometry);
+
+	// FWidgetReflectorNodeBase interface
+	virtual EWidgetReflectorNodeType GetNodeType() const override;
+	virtual TSharedPtr<SWidget> GetLiveWidget() const override;
 	virtual FText GetWidgetType() const override;
 	virtual FText GetWidgetVisibilityText() const override;
 	virtual FText GetWidgetReadableLocation() const override;
@@ -165,45 +241,52 @@ public:
 	virtual FString GetWidgetAddress() const override;
 	virtual bool GetWidgetEnabled() const override;
 
+	/** Save this node data as a JSON object */
+	static TSharedRef<FJsonValue> ToJson(const TSharedRef<FSnapshotWidgetReflectorNode>& RootSnapshotNode);
+
+	/** Populate this node data from a JSON object */
+	static TSharedRef<FSnapshotWidgetReflectorNode> FromJson(const TSharedRef<FJsonValue>& RootJsonValue);
+
+private:
+	/**
+	 * Default constructor
+	 */
+	FSnapshotWidgetReflectorNode();
+
+	/**
+	 * Construct this node from the given widget geometry, caching out any data that may be required for future visualization in the widget reflector
+	 */
+	explicit FSnapshotWidgetReflectorNode(const FArrangedWidget& InWidgetGeometry);
+
 private:
 	/** The type string of the widget at the point it was passed to Initialize */
-	UPROPERTY()
 	FText CachedWidgetType;
 
 	/** The visibility string of the widget at the point it was passed to Initialize */
-	UPROPERTY()
 	FText CachedWidgetVisibilityText;
 
 	/** The human readable location (source file for C++ widgets, asset name for UMG widgets) of the widget at the point it was passed to Initialize */
-	UPROPERTY()
 	FText CachedWidgetReadableLocation;
 
 	/** The name of the file that the widget was created from at the point it was passed to Initialize (for C++ widgets) */
-	UPROPERTY()
 	FString CachedWidgetFile;
 
 	/** The line number of the file that the widget was created from at the point it was passed to Initialize (for C++ widgets) */
-	UPROPERTY()
 	int32 CachedWidgetLineNumber;
 
 	/** The name of the asset that the widget was created from at the point it was passed to Initialize (for UMG widgets) */
-	UPROPERTY()
 	FName CachedWidgetAssetName;
 
 	/** The desired size of the widget at the point it was passed to Initialize */
-	UPROPERTY()
 	FVector2D CachedWidgetDesiredSize;
 
 	/** The foreground color of the widget at the point it was passed to Initialize */
-	UPROPERTY()
 	FSlateColor CachedWidgetForegroundColor;
 
 	/** The in-memory address of the widget at the point it was passed to Initialize */
-	UPROPERTY()
 	FString CachedWidgetAddress;
 
 	/** The enabled state of the widget at the point it was passed to Initialize */
-	UPROPERTY()
 	bool CachedWidgetEnabled;
 };
 
@@ -217,7 +300,7 @@ public:
 	 * @param InNodeClass The type of widget reflector node to create
 	 * @param InWidgetGeometry Optional widget and associated geometry which this node should represent
 	 */
-	static ULiveWidgetReflectorNode* NewLiveNode(const FArrangedWidget& InWidgetGeometry = FArrangedWidget(SNullWidget::NullWidget, FGeometry()));
+	static TSharedRef<FLiveWidgetReflectorNode> NewLiveNode(const FArrangedWidget& InWidgetGeometry = FArrangedWidget(SNullWidget::NullWidget, FGeometry()));
 
 	/**
 	 * Create nodes for the supplied widget and all their children such that they reference a live widget
@@ -226,7 +309,7 @@ public:
 	 * @param InNodeClass The type of widget reflector node to create
 	 * @param InWidgetGeometry Widget and geometry whose children to capture in the snapshot.
 	 */
-	static ULiveWidgetReflectorNode* NewLiveNodeTreeFrom(const FArrangedWidget& InWidgetGeometry);
+	static TSharedRef<FLiveWidgetReflectorNode> NewLiveNodeTreeFrom(const FArrangedWidget& InWidgetGeometry);
 
 	/**
 	 * Create a single node referencing a snapshot of its current state.
@@ -234,7 +317,7 @@ public:
 	 * @param InNodeClass The type of widget reflector node to create
 	 * @param InWidgetGeometry Optional widget and associated geometry which this node should represent
 	 */
-	static USnapshotWidgetReflectorNode* NewSnapshotNode(const FArrangedWidget& InWidgetGeometry = FArrangedWidget(SNullWidget::NullWidget, FGeometry()));
+	static TSharedRef<FSnapshotWidgetReflectorNode> NewSnapshotNode(const FArrangedWidget& InWidgetGeometry = FArrangedWidget(SNullWidget::NullWidget, FGeometry()));
 
 	/**
 	 * Create nodes for the supplied widget and all their children such that they reference a snapshot of their current state
@@ -243,7 +326,7 @@ public:
 	 * @param InNodeClass The type of widget reflector node to create
 	 * @param InWidgetGeometry Widget and geometry whose children to capture in the snapshot.
 	 */
-	static USnapshotWidgetReflectorNode* NewSnapshotNodeTreeFrom(const FArrangedWidget& InWidgetGeometry);
+	static TSharedRef<FSnapshotWidgetReflectorNode> NewSnapshotNodeTreeFrom(const FArrangedWidget& InWidgetGeometry);
 
 private:
 	/**
@@ -252,7 +335,7 @@ private:
 	 * @param InNodeClass The type of widget reflector node to create
 	 * @param InWidgetGeometry Optional widget and associated geometry which this node should represent
 	 */
-	static UWidgetReflectorNodeBase* NewNode(const TSubclassOf<UWidgetReflectorNodeBase>& InNodeClass, const FArrangedWidget& InWidgetGeometry = FArrangedWidget(SNullWidget::NullWidget, FGeometry()));
+	static TSharedRef<FWidgetReflectorNodeBase> NewNode(const EWidgetReflectorNodeType InNodeType, const FArrangedWidget& InWidgetGeometry = FArrangedWidget(SNullWidget::NullWidget, FGeometry()));
 
 	/**
 	 * Create nodes for the supplied widget and all their children
@@ -261,7 +344,7 @@ private:
 	 * @param InNodeClass The type of widget reflector node to create
 	 * @param InWidgetGeometry Widget and geometry whose children to capture in the snapshot.
 	 */
-	static UWidgetReflectorNodeBase* NewNodeTreeFrom(const TSubclassOf<UWidgetReflectorNodeBase>& InNodeClass, const FArrangedWidget& InWidgetGeometry);
+	static TSharedRef<FWidgetReflectorNodeBase> NewNodeTreeFrom(const EWidgetReflectorNodeType InNodeType, const FArrangedWidget& InWidgetGeometry);
 
 public:
 	/**
@@ -273,7 +356,7 @@ public:
 	 * @param SearchResult An array that gets results put in it
 	 * @param NodeIndexToFind Index of the widget in the path that we are currently looking for; we are done when we've found all of them
 	 */
-	static void FindLiveWidgetPath(const TArray<UWidgetReflectorNodeBase*>& CandidateNodes, const FWidgetPath& WidgetPathToFind, TArray<UWidgetReflectorNodeBase*>& SearchResult, int32 NodeIndexToFind = 0);
+	static void FindLiveWidgetPath(const TArray<TSharedRef<FWidgetReflectorNodeBase>>& CandidateNodes, const FWidgetPath& WidgetPathToFind, TArray<TSharedRef<FWidgetReflectorNodeBase>>& SearchResult, int32 NodeIndexToFind = 0);
 
 public:
 	/**
diff --git a/Engine/Source/Developer/SlateReflector/Private/SlateReflectorModule.cpp b/Engine/Source/Developer/SlateReflector/Private/SlateReflectorModule.cpp
--- a/Engine/Source/Developer/SlateReflector/Private/SlateReflectorModule.cpp
+++ b/Engine/Source/Developer/SlateReflector/Private/SlateReflectorModule.cpp
@@ -20,13 +20,14 @@ public:
 
 	// ISlateReflectorModule interface
 
-	virtual TSharedRef<SWidget> GetWidgetReflector() override
+	TSharedRef<SWidget> GetWidgetReflector(const TSharedRef<SDockTab>& InParentTab)
 	{
 		TSharedPtr<SWidgetReflector> WidgetReflector = WidgetReflectorPtr.Pin();
 
 		if (!WidgetReflector.IsValid())
 		{
-			WidgetReflector = SNew(SWidgetReflector);
+			WidgetReflector = SNew(SWidgetReflector)
+				.ParentTab(InParentTab);
 			WidgetReflectorPtr = WidgetReflector;
 			FSlateApplication::Get().SetWidgetReflector(WidgetReflector.ToSharedRef());
 		}
@@ -34,13 +35,13 @@ public:
 		return WidgetReflector.ToSharedRef();
 	}
 
-	virtual TSharedRef<SWidget> GetAtlasVisualizer( ISlateAtlasProvider* InAtlasProvider ) override
+	TSharedRef<SWidget> GetAtlasVisualizer( ISlateAtlasProvider* InAtlasProvider )
 	{
 		return SNew(SAtlasVisualizer)
 			.AtlasProvider(InAtlasProvider);
 	}
 
-	virtual TSharedRef<SWidget> GetTextureAtlasVisualizer() override
+	TSharedRef<SWidget> GetTextureAtlasVisualizer()
 	{
 		ISlateAtlasProvider* const AtlasProvider = FSlateApplication::Get().GetRenderer()->GetTextureAtlasProvider();
 
@@ -60,7 +61,7 @@ public:
 		}
 	}
 
-	virtual TSharedRef<SWidget> GetFontAtlasVisualizer() override
+	TSharedRef<SWidget> GetFontAtlasVisualizer()
 	{
 		ISlateAtlasProvider* const AtlasProvider = FSlateApplication::Get().GetRenderer()->GetFontAtlasProvider();
 
@@ -82,91 +83,66 @@ public:
 
 	virtual void DisplayWidgetReflector() override
 	{
-		if( bHasRegisteredTabSpawners )
-		{
-			FGlobalTabmanager::Get()->InvokeTab(FTabId("WidgetReflector"));
-		}
-		else if( !WidgetReflectorWindowPtr.IsValid() )
-		{
-			const TSharedRef<SWindow> ReflectorWindow = SNew(SWindow)
-				.AutoCenter(EAutoCenter::PrimaryWorkArea)
-				.ClientSize(FVector2D(600,400))
-				.Title(LOCTEXT("WidgetReflectorTitle", "Widget Reflector"))
-				[
-					GetWidgetReflector()
-				];
-		
-			WidgetReflectorWindowPtr = ReflectorWindow;
-		
-			FSlateApplication::Get().AddWindow( ReflectorWindow );
-		}
+		check(bHasRegisteredTabSpawners);
+		FGlobalTabmanager::Get()->InvokeTab(FTabId("WidgetReflector"));
 	}
 
 	virtual void DisplayTextureAtlasVisualizer() override
 	{
-		if( bHasRegisteredTabSpawners )
-		{
-			FGlobalTabmanager::Get()->InvokeTab(FTabId("TextureAtlasVisualizer"));
-		}
-		else if( !TextureAtlasVisualizerWindowPtr.IsValid() )
-		{
-			const TSharedRef<SWindow> AtlasVisWindow = SNew(SWindow)
-				.AutoCenter(EAutoCenter::PrimaryWorkArea)
-				.ClientSize(FVector2D(1040, 900))
-				.Title(LOCTEXT("TextureAtlasVisualizerTitle", "Texture Atlas Visualizer"))
-				[
-					GetTextureAtlasVisualizer()
-				];
-		
-			TextureAtlasVisualizerWindowPtr = AtlasVisWindow;
-		
-			FSlateApplication::Get().AddWindow( AtlasVisWindow );
-		}
+		check(bHasRegisteredTabSpawners);
+		FGlobalTabmanager::Get()->InvokeTab(FTabId("TextureAtlasVisualizer"));
 	}
 
 	virtual void DisplayFontAtlasVisualizer() override
 	{
-		if( bHasRegisteredTabSpawners )
+		check(bHasRegisteredTabSpawners);
+		FGlobalTabmanager::Get()->InvokeTab(FTabId("FontAtlasVisualizer"));
+	}
+
+	virtual void RegisterTabSpawner( const TSharedPtr<FWorkspaceItem>& WorkspaceGroup ) override
+	{
+		if (bHasRegisteredTabSpawners)
 		{
-			FGlobalTabmanager::Get()->InvokeTab(FTabId("FontAtlasVisualizer"));
+			UnregisterTabSpawner();
 		}
-		else if( !FontAtlasVisualizerWindowPtr.IsValid() )
+
+		bHasRegisteredTabSpawners = true;
+
 		{
-			const TSharedRef<SWindow> AtlasVisWindow = SNew(SWindow)
-				.AutoCenter(EAutoCenter::PrimaryWorkArea)
-				.ClientSize(FVector2D(1040, 900))
-				.Title(LOCTEXT("FontAtlasVisualizerTitle", "Font Atlas Visualizer"))
-				[
-					GetFontAtlasVisualizer()
-				];
-		
-			FontAtlasVisualizerWindowPtr = AtlasVisWindow;
-		
-			FSlateApplication::Get().AddWindow( AtlasVisWindow );
+			FTabSpawnerEntry& SpawnerEntry = FGlobalTabmanager::Get()->RegisterNomadTabSpawner("WidgetReflector", FOnSpawnTab::CreateRaw(this, &FSlateReflectorModule::MakeWidgetReflectorTab) )
+				.SetDisplayName(LOCTEXT("WidgetReflectorTitle", "Widget Reflector"))
+				.SetTooltipText(LOCTEXT("WidgetReflectorTooltipText", "Open the Widget Reflector tab."))
+				.SetIcon(FSlateIcon(FCoreStyle::Get().GetStyleSetName(), "WidgetReflector.TabIcon"));
+
+			if (WorkspaceGroup.IsValid())
+			{
+				SpawnerEntry.SetGroup(WorkspaceGroup.ToSharedRef());
+			}
 		}
-	}
 
-	virtual void RegisterTabSpawner( const TSharedRef<FWorkspaceItem>& WorkspaceGroup ) override
-	{
-		bHasRegisteredTabSpawners = true;
+		{
+			FTabSpawnerEntry& SpawnerEntry = FGlobalTabmanager::Get()->RegisterNomadTabSpawner("TextureAtlasVisualizer", FOnSpawnTab::CreateRaw(this, &FSlateReflectorModule::MakeTextureAtlasVisualizerTab) )
+				.SetDisplayName(LOCTEXT("TextureAtlasVisualizerTitle", "Texture Atlas Visualizer"))
+				.SetTooltipText(LOCTEXT("TextureAtlasVisualizerTooltipText", "Open the Texture Atlas Visualizer tab."))
+				.SetMenuType(ETabSpawnerMenuType::Hidden);
+
+			if (WorkspaceGroup.IsValid())
+			{
+				SpawnerEntry.SetGroup(WorkspaceGroup.ToSharedRef());
+			}
+		}
 
-		FGlobalTabmanager::Get()->RegisterNomadTabSpawner("WidgetReflector", FOnSpawnTab::CreateRaw(this, &FSlateReflectorModule::MakeWidgetReflectorTab) )
-			.SetDisplayName(LOCTEXT("WidgetReflectorTitle", "Widget Reflector"))
-			.SetTooltipText(LOCTEXT("WidgetReflectorTooltipText", "Open the Widget Reflector tab."))
-			.SetGroup(WorkspaceGroup)
-			.SetIcon(FSlateIcon(FCoreStyle::Get().GetStyleSetName(), "WidgetReflector.TabIcon"));
-
-		FGlobalTabmanager::Get()->RegisterNomadTabSpawner("TextureAtlasVisualizer", FOnSpawnTab::CreateRaw(this, &FSlateReflectorModule::MakeTextureAtlasVisualizerTab) )
-			.SetDisplayName(LOCTEXT("TextureAtlasVisualizerTitle", "Texture Atlas Visualizer"))
-			.SetTooltipText(LOCTEXT("TextureAtlasVisualizerTooltipText", "Open the Texture Atlas Visualizer tab."))
-			.SetGroup(WorkspaceGroup)
-			.SetMenuType(ETabSpawnerMenuType::Hidden);
-
-		FGlobalTabmanager::Get()->RegisterNomadTabSpawner("FontAtlasVisualizer", FOnSpawnTab::CreateRaw(this, &FSlateReflectorModule::MakeFontAtlasVisualizerTab) )
-			.SetDisplayName(LOCTEXT("FontAtlasVisualizerTitle", "Font Atlas Visualizer"))
-			.SetTooltipText(LOCTEXT("FontAtlasVisualizerTooltipText", "Open the Font Atlas Visualizer tab."))
-			.SetGroup(WorkspaceGroup)
-			.SetMenuType(ETabSpawnerMenuType::Hidden);
+		{
+			FTabSpawnerEntry& SpawnerEntry = FGlobalTabmanager::Get()->RegisterNomadTabSpawner("FontAtlasVisualizer", FOnSpawnTab::CreateRaw(this, &FSlateReflectorModule::MakeFontAtlasVisualizerTab) )
+				.SetDisplayName(LOCTEXT("FontAtlasVisualizerTitle", "Font Atlas Visualizer"))
+				.SetTooltipText(LOCTEXT("FontAtlasVisualizerTooltipText", "Open the Font Atlas Visualizer tab."))
+				.SetMenuType(ETabSpawnerMenuType::Hidden);
+
+			if (WorkspaceGroup.IsValid())
+			{
+				SpawnerEntry.SetGroup(WorkspaceGroup.ToSharedRef());
+			}
+		}
 	}
 
 	virtual void UnregisterTabSpawner() override
@@ -185,6 +161,7 @@ public:
 	virtual void StartupModule() override
 	{
 		bHasRegisteredTabSpawners = false;
+		RegisterTabSpawner(nullptr);
 	}
 
 	virtual void ShutdownModule() override
@@ -196,11 +173,10 @@ private:
 
 	TSharedRef<SDockTab> MakeWidgetReflectorTab( const FSpawnTabArgs& )
 	{
-		return SNew(SDockTab)
-			.TabRole(ETabRole::NomadTab)
-			[
-				GetWidgetReflector()
-			];
+		TSharedRef<SDockTab> WidgetReflectorTab = SNew(SDockTab)
+			.TabRole(ETabRole::NomadTab);
+		WidgetReflectorTab->SetContent(GetWidgetReflector(WidgetReflectorTab));
+		return WidgetReflectorTab;
 	}
 
 	TSharedRef<SDockTab> MakeTextureAtlasVisualizerTab( const FSpawnTabArgs& )
@@ -228,15 +204,6 @@ private:
 
 	/** Holds the widget reflector singleton. */
 	TWeakPtr<SWidgetReflector> WidgetReflectorPtr;
-
-	/** Holds the widget reflector singleton. */
-	TWeakPtr<SWindow> WidgetReflectorWindowPtr;
-
-	/** Holds the texture atlas viewer singleton. */
-	TWeakPtr<SWindow> TextureAtlasVisualizerWindowPtr;
-
-	/** Holds the font atlas viewer singleton. */
-	TWeakPtr<SWindow> FontAtlasVisualizerWindowPtr;
 };
 
 
diff --git a/Engine/Source/Developer/SlateReflector/Private/SlateReflectorPrivatePCH.h b/Engine/Source/Developer/SlateReflector/Private/SlateReflectorPrivatePCH.h
--- a/Engine/Source/Developer/SlateReflector/Private/SlateReflectorPrivatePCH.h
+++ b/Engine/Source/Developer/SlateReflector/Private/SlateReflectorPrivatePCH.h
@@ -11,6 +11,7 @@
 #include "InputCore.h"
 #include "SlateBasics.h"
 #include "SlateCore.h"
+#include "Json.h"
 
 
 /* Private includes
diff --git a/Engine/Source/Developer/SlateReflector/Private/Widgets/SWidgetReflector.cpp b/Engine/Source/Developer/SlateReflector/Private/Widgets/SWidgetReflector.cpp
--- a/Engine/Source/Developer/SlateReflector/Private/Widgets/SWidgetReflector.cpp
+++ b/Engine/Source/Developer/SlateReflector/Private/Widgets/SWidgetReflector.cpp
@@ -4,6 +4,12 @@
 #include "ISlateReflectorModule.h"
 #include "SlateStats.h"
 #include "SInvalidationPanel.h"
+#include "SDockTab.h"
+#include "SWidgetSnapshotVisualizer.h"
+
+#if SLATE_REFLECTOR_HAS_DESKTOP_PLATFORM
+#include "DesktopPlatformModule.h"
+#endif // SLATE_REFLECTOR_HAS_DESKTOP_PLATFORM
 
 #define LOCTEXT_NAMESPACE "SWidgetReflector"
 #define WITH_EVENT_LOGGING 0
@@ -47,18 +53,48 @@ struct FLoggedEvent
 namespace WidgetReflectorImpl
 {
 
+/** Different UI modes the widget reflector can be in */
+enum class EWidgetReflectorUIMode : uint8
+{
+	Live,
+	Snapshot,
+};
+
+namespace WidgetReflectorTabID
+{
+	static const FName WidgetHierarchy = "WidgetReflector.WidgetHierarchyTab";
+	static const FName SlateStats = "WidgetReflector.SlateStatsTab";
+	static const FName SnapshotWidgetPicker = "WidgetReflector.SnapshotWidgetPickerTab";
+}
+
 /**
  * Widget reflector implementation
  */
-class SWidgetReflector : public ::SWidgetReflector, public FGCObject
+class SWidgetReflector : public ::SWidgetReflector
 {
-	// The reflector uses a tree that observes UWidgetReflectorNodeBase objects.
-	typedef STreeView<UWidgetReflectorNodeBase*> SReflectorTree;
+	// The reflector uses a tree that observes FWidgetReflectorNodeBase objects.
+	typedef STreeView<TSharedRef<FWidgetReflectorNodeBase>> SReflectorTree;
+
+public:
+
+	~SWidgetReflector();
 
 private:
 
 	virtual void Construct( const FArguments& InArgs ) override;
 
+	TSharedRef<SDockTab> SpawnWidgetHierarchyTab(const FSpawnTabArgs& Args);
+#if SLATE_STATS
+	TSharedRef<SDockTab> SpawnSlateStatsTab(const FSpawnTabArgs& Args);
+#endif
+	TSharedRef<SDockTab> SpawnSnapshotWidgetPicker(const FSpawnTabArgs& Args);
+
+	void OnTabSpawned(const FName& TabIdentifier, const TSharedRef<SDockTab>& SpawnedTab);
+
+	void CloseTab(const FName& TabIdentifier);
+
+	void SetUIMode(const EWidgetReflectorUIMode InNewMode);
+
 	// SCompoundWidget overrides
 	virtual void Tick( const FGeometry& AllottedGeometry, const double InCurrentTime, const float InDeltaTime ) override;
 
@@ -108,14 +144,14 @@ private:
 	 * @param InReflectorNode The node to generate the tool tip for.
 	 * @return The tool tip widget.
 	 */
-	TSharedRef<SToolTip> GenerateToolTipForReflectorNode( UWidgetReflectorNodeBase* InReflectorNode );
+	TSharedRef<SToolTip> GenerateToolTipForReflectorNode( TSharedRef<FWidgetReflectorNodeBase> InReflectorNode );
 
 	/**
 	 * Mark the provided reflector nodes such that they stand out in the tree and are visible.
 	 *
 	 * @param WidgetPathToObserve The nodes to mark.
 	 */
-	void VisualizeAsTree( const TArray< UWidgetReflectorNodeBase* >& WidgetPathToVisualize );
+	void VisualizeAsTree( const TArray< TSharedRef<FWidgetReflectorNodeBase> >& WidgetPathToVisualize );
 
 	/**
 	 * Draw the widget path to the picked widget as the widgets' outlines.
@@ -136,7 +172,7 @@ private:
 	 * @param LayerId the maximum layer achieved in OutDrawElements so far.
 	 * @return The maximum layer ID we achieved while painting.
 	 */
-	int32 VisualizeSelectedNodesAsRectangles( const TArray<UWidgetReflectorNodeBase*>& InNodesToDraw, const TSharedRef<SWindow>& VisualizeInWindow, FSlateWindowElementList& OutDrawElements, int32 LayerId );
+	int32 VisualizeSelectedNodesAsRectangles( const TArray<TSharedRef<FWidgetReflectorNodeBase>>& InNodesToDraw, const TSharedRef<SWindow>& VisualizeInWindow, FSlateWindowElementList& OutDrawElements, int32 LayerId );
 
 	/** Callback for changing the application scale slider. */
 	void HandleAppScaleSliderChanged( float NewValue )
@@ -173,6 +209,7 @@ private:
 		if (bIsPicking)
 		{
 			bShowFocus = false;
+			SetUIMode(EWidgetReflectorUIMode::Live);
 		}
 
 		return FReply::Handled();
@@ -182,50 +219,53 @@ private:
 	FSlateColor HandlePickButtonColorAndOpacity() const
 	{
 		static const FName SelectionColor("SelectionColor");
-		static const FName DefaultForeground("DefaultForeground");
 
 		return bIsPicking
 			? FCoreStyle::Get().GetSlateColor(SelectionColor)
-			: FCoreStyle::Get().GetSlateColor(DefaultForeground);
+			: FLinearColor::White;
 	}
 
 	/** Callback for getting the text of the pick button. */
 	FText HandlePickButtonText() const;
 
-	/** Callback for clicking the snapshot button. */
-	FReply HandleSnapshotButtonClicked();
+	/** Callback for clicking the "Take Snapshot" button. */
+	FReply HandleTakeSnapshotButtonClicked();
+
+#if SLATE_REFLECTOR_HAS_DESKTOP_PLATFORM
+	/** Callback for clicking the "Load Snapshot" button. */
+	FReply HandleLoadSnapshotButtonClicked();
+#endif // SLATE_REFLECTOR_HAS_DESKTOP_PLATFORM
 
 	/** Callback for generating a row in the reflector tree view. */
-	TSharedRef<ITableRow> HandleReflectorTreeGenerateRow( UWidgetReflectorNodeBase* InReflectorNode, const TSharedRef<STableViewBase>& OwnerTable );
+	TSharedRef<ITableRow> HandleReflectorTreeGenerateRow( TSharedRef<FWidgetReflectorNodeBase> InReflectorNode, const TSharedRef<STableViewBase>& OwnerTable );
 
 	/** Callback for getting the child items of the given reflector tree node. */
-	void HandleReflectorTreeGetChildren( UWidgetReflectorNodeBase* InWidgetGeometry, TArray<UWidgetReflectorNodeBase*>& OutChildren );
+	void HandleReflectorTreeGetChildren( TSharedRef<FWidgetReflectorNodeBase> InReflectorNode, TArray<TSharedRef<FWidgetReflectorNodeBase>>& OutChildren );
 
 	/** Callback for when the selection in the reflector tree has changed. */
-	void HandleReflectorTreeSelectionChanged( UWidgetReflectorNodeBase*, ESelectInfo::Type /*SelectInfo*/ )
+	void HandleReflectorTreeSelectionChanged( TSharedPtr<FWidgetReflectorNodeBase>, ESelectInfo::Type /*SelectInfo*/ )
 	{
 		SelectedNodes = ReflectorTree->GetSelectedItems();
 	}
 
 	TSharedRef<ITableRow> GenerateEventLogRow( TSharedRef<FLoggedEvent> InReflectorNode, const TSharedRef<STableViewBase>& OwnerTable );
 
-	virtual void AddReferencedObjects( FReferenceCollector& Collector ) override
-	{
-		// STreeView only references items that are currently visible, but we need to prevent the entire tree from being GC'd
-		for (UWidgetReflectorNodeBase* NodeRootItem : ReflectorTreeRoot)
-		{
-			Collector.AddReferencedObject(NodeRootItem);
-		}
-	}
+	EWidgetReflectorUIMode CurrentUIMode;
+
+	TSharedPtr<FTabManager> TabManager;
+	TMap<FName, TWeakPtr<SDockTab>> SpawnedTabs;
 
 	TArray< TSharedRef<FLoggedEvent> > LoggedEvents;
 	TSharedPtr< SListView< TSharedRef< FLoggedEvent > > > EventListView;
-	TSharedPtr< SHorizontalBox > StatsToolsBox;
 	TSharedPtr<SReflectorTree> ReflectorTree;
 
-	TArray<UWidgetReflectorNodeBase*> SelectedNodes;
-	TArray<UWidgetReflectorNodeBase*> ReflectorTreeRoot;
-	TArray<UWidgetReflectorNodeBase*> PickedPath;
+	TArray<TSharedRef<FWidgetReflectorNodeBase>> SelectedNodes;
+	TArray<TSharedRef<FWidgetReflectorNodeBase>> ReflectorTreeRoot;
+	TArray<TSharedRef<FWidgetReflectorNodeBase>> PickedPath;
+
+	/** When working with a snapshotted tree, this will contain the snapshot hierarchy and screenshot info */
+	FWidgetSnapshotData SnapshotData;
+	TSharedPtr<SWidgetSnapshotVisualizer> WidgetSnapshotVisualizer;
 
 	SSplitter::FSlot* WidgetInfoLocation;
 
@@ -253,24 +293,20 @@ private:
 };
 
 
-
-
-
-
-
-
-
-
-
-
-
-
+SWidgetReflector::~SWidgetReflector()
+{
+	TabManager->UnregisterTabSpawner(WidgetReflectorTabID::WidgetHierarchy);
+	TabManager->UnregisterTabSpawner(WidgetReflectorTabID::SlateStats);
+	TabManager->UnregisterTabSpawner(WidgetReflectorTabID::SnapshotWidgetPicker);
+}
 
 
 void SWidgetReflector::Construct( const FArguments& InArgs )
 {
 	LoggedEvents.Reserve(MaxLoggedEvents);
 
+	CurrentUIMode = EWidgetReflectorUIMode::Live;
+
 	bShowFocus = false;
 	bIsPicking = false;
 
@@ -278,332 +314,502 @@ void SWidgetReflector::Construct( const FArguments& InArgs )
 	LastMouseClickTime = -1;
 	CursorPingPosition = FVector2D::ZeroVector;
 
+#if SLATE_STATS
+	const FName TabLayoutName = "WidgetReflector_Layout_v1";
+#else
+	const FName TabLayoutName = "WidgetReflector_Layout_NoStats_v1";
+#endif
+
+	TSharedRef<FTabManager::FLayout> Layout = FTabManager::NewLayout(TabLayoutName)
+	->AddArea
+	(
+		FTabManager::NewPrimaryArea()
+		->SetOrientation(Orient_Vertical)
+		->Split
+		(
+			FTabManager::NewStack()
+			->SetHideTabWell(true)
+			->SetSizeCoefficient(0.7f)
+			->AddTab(WidgetReflectorTabID::WidgetHierarchy, ETabState::OpenedTab)
+		)
+		->Split
+		(
+			FTabManager::NewStack()
+			->SetHideTabWell(true)
+			->SetSizeCoefficient(0.3f)
+#if SLATE_STATS
+			->AddTab(WidgetReflectorTabID::SlateStats, ETabState::ClosedTab)
+#endif
+			->AddTab(WidgetReflectorTabID::SnapshotWidgetPicker, ETabState::ClosedTab)
+		)
+	);
+
+	auto RegisterTrackedTabSpawner = [this](const FName& TabId, const FOnSpawnTab& OnSpawnTab) -> FTabSpawnerEntry&
+	{
+		return TabManager->RegisterTabSpawner(TabId, FOnSpawnTab::CreateLambda([this, OnSpawnTab](const FSpawnTabArgs& Args) -> TSharedRef<SDockTab>
+		{
+			TSharedRef<SDockTab> SpawnedTab = OnSpawnTab.Execute(Args);
+			OnTabSpawned(Args.GetTabId().TabType, SpawnedTab);
+			return SpawnedTab;
+		}));
+	};
+
+	check(InArgs._ParentTab.IsValid());
+	TabManager = FGlobalTabmanager::Get()->NewTabManager(InArgs._ParentTab.ToSharedRef());
+
+	RegisterTrackedTabSpawner(WidgetReflectorTabID::WidgetHierarchy, FOnSpawnTab::CreateSP(this, &SWidgetReflector::SpawnWidgetHierarchyTab))
+		.SetDisplayName(LOCTEXT("WidgetHierarchyTab", "Widget Hierarchy"));
+
+#if SLATE_STATS
+	RegisterTrackedTabSpawner(WidgetReflectorTabID::SlateStats, FOnSpawnTab::CreateSP(this, &SWidgetReflector::SpawnSlateStatsTab))
+		.SetDisplayName(LOCTEXT("SlateStatsTab", "Slate Stats"));
+#endif
+
+	RegisterTrackedTabSpawner(WidgetReflectorTabID::SnapshotWidgetPicker, FOnSpawnTab::CreateSP(this, &SWidgetReflector::SpawnSnapshotWidgetPicker))
+		.SetDisplayName(LOCTEXT("SnapshotWidgetPickerTab", "Snapshot Widget Picker"));
+
 	this->ChildSlot
 	[
 		SNew(SBorder)
-			.BorderImage(FCoreStyle::Get().GetBrush("ToolPanel.GroupBorder"))
-			[
-				SNew(SVerticalBox)
+		.BorderImage(FCoreStyle::Get().GetBrush("ToolPanel.GroupBorder"))
+		.BorderBackgroundColor(FLinearColor::Gray) // Darken the outer border
+		[
+			SNew(SVerticalBox)
 				
-				+ SVerticalBox::Slot()
-					.AutoHeight()
-					[
-						SNew(SHorizontalBox)
+			+SVerticalBox::Slot()
+			.AutoHeight()
+			.Padding(FMargin(0.0f, 0.0f, 0.0f, 6.0f))
+			[
+				SNew(SHorizontalBox)
 					
-						+ SHorizontalBox::Slot()
-						.AutoWidth()
-						.VAlign(VAlign_Center)
-						[
-							SNew(STextBlock)
-								.Text(LOCTEXT("AppScale", "Application Scale: "))
-						]
+				+SHorizontalBox::Slot()
+				.AutoWidth()
+				.VAlign(VAlign_Center)
+				[
+					SNew(STextBlock)
+					.Text(LOCTEXT("AppScale", "Application Scale: "))
+				]
 
-						+ SHorizontalBox::Slot()
-						.AutoWidth()
-						[
-							SNew(SBox)
-							.MinDesiredWidth(100)
-							.MaxDesiredWidth(250)
-							[
-								SNew(SSpinBox<float>)
-									.Value(this, &SWidgetReflector::HandleAppScaleSliderValue)
-									.MinValue(0.1f)
-									.MaxValue(3.0f)
-									.Delta(0.01f)
-									.OnValueChanged(this, &SWidgetReflector::HandleAppScaleSliderChanged)
-							]
-						]
-						+ SHorizontalBox::Slot()
-						.FillWidth(1.0f)
-						[
-							SNew(SSpacer)
-						]
-						+ SHorizontalBox::Slot()
-						.AutoWidth()
-						.Padding(FMargin(5.0f, 0.0f))
-						[
-							SNew(SButton)
-							.Text(LOCTEXT("DisplayTextureAtlases", "Display Texture Atlases"))
-							.OnClicked(this, &SWidgetReflector::HandleDisplayTextureAtlases)
-						]
+				+SHorizontalBox::Slot()
+				.AutoWidth()
+				[
+					SNew(SBox)
+					.MinDesiredWidth(100)
+					.MaxDesiredWidth(250)
+					[
+						SNew(SSpinBox<float>)
+						.Value(this, &SWidgetReflector::HandleAppScaleSliderValue)
+						.MinValue(0.1f)
+						.MaxValue(3.0f)
+						.Delta(0.01f)
+						.OnValueChanged(this, &SWidgetReflector::HandleAppScaleSliderChanged)
+					]
+				]
 
-						+ SHorizontalBox::Slot()
-						.AutoWidth()
-						.Padding(FMargin(5.0f, 0.0f))
+				+SHorizontalBox::Slot()
+				.AutoWidth()
+				.Padding(FMargin(5.0f, 0.0f))
+				[
+					SNew(SCheckBox)
+					.Style( FCoreStyle::Get(), "ToggleButtonCheckbox" )
+					.IsChecked_Lambda([]()
+					{
+						return bFoldTick == 0 ? ECheckBoxState::Unchecked : ECheckBoxState::Checked;
+					})
+					.OnCheckStateChanged_Lambda([]( const ECheckBoxState NewState )
+					{
+						bFoldTick = (NewState == ECheckBoxState::Checked) ? 1 : 0;
+					})
+					[
+						SNew(SBox)
+						.VAlign( VAlign_Center )
+						.HAlign( HAlign_Center )
+						.Padding(FMargin(4.0, 2.0))
 						[
-							SNew(SButton)
-							.Text(LOCTEXT("DisplayFontAtlases", "Display Font Atlases"))
-							.OnClicked(this, &SWidgetReflector::HandleDisplayFontAtlases)
+							SNew(STextBlock)	
+							.Text( LOCTEXT("ToggleTickFolding", "Fold Tick") )
 						]
 					]
+				]
 
-				+ SVerticalBox::Slot()
-				.AutoHeight()
+#if !(UE_BUILD_SHIPPING || UE_BUILD_TEST)
+				+SHorizontalBox::Slot()
+				.AutoWidth()
+				.Padding(FMargin(5.0f, 0.0f))
 				[
-					SNew(SHorizontalBox)
-
-					+ SHorizontalBox::Slot()
-					.AutoWidth()
-					.Padding(5.0f)
+					SNew(SCheckBox)
+					.Style(FCoreStyle::Get(), "ToggleButtonCheckbox")
+					.IsChecked_Lambda([&]()
+					{
+						return SInvalidationPanel::IsInvalidationDebuggingEnabled() ? ECheckBoxState::Checked : ECheckBoxState::Unchecked;
+					})
+					.OnCheckStateChanged_Lambda([&](const ECheckBoxState NewState)
+					{
+						SInvalidationPanel::EnableInvalidationDebugging(( NewState == ECheckBoxState::Checked ) ? true : false);
+					})
 					[
-						// Check box that controls LIVE MODE
-						SNew(SCheckBox)
-						.IsChecked(this, &SWidgetReflector::HandleFocusCheckBoxIsChecked)
-						.OnCheckStateChanged(this, &SWidgetReflector::HandleFocusCheckBoxCheckedStateChanged)
+						SNew(SBox)
+						.VAlign(VAlign_Center)
+						.HAlign(HAlign_Center)
+						.Padding(FMargin(4.0, 2.0))
 						[
 							SNew(STextBlock)
-							.Text(LOCTEXT("ShowFocus", "Show Focus"))
+							.Text(LOCTEXT("InvalidationDebugging", "Invalidation Debugging"))
 						]
 					]
+				]
+#endif
 
-					+ SHorizontalBox::Slot()
-					.AutoWidth()
-					.Padding(5.0f)
-					[
-						// Check box that controls PICKING A WIDGET TO INSPECT
-						SNew(SButton)
-							.OnClicked(this, &SWidgetReflector::HandlePickButtonClicked)
-							.ButtonColorAndOpacity(this, &SWidgetReflector::HandlePickButtonColorAndOpacity)
-							[
-								SNew(STextBlock)
-									.Text(this, &SWidgetReflector::HandlePickButtonText)
-							]
-					]
-
-					/*
-					+ SHorizontalBox::Slot()
-					.AutoWidth()
-					.Padding(5.0f)
-					[
-						// Button that controls taking a snapshot of the current window(s)
-						SNew(SButton)
-							.OnClicked(this, &SWidgetReflector::HandleSnapshotButtonClicked)
-							[
-								SNew(STextBlock)
-									.Text(LOCTEXT("SnapshotButtonText", "Snapshot"))
-							]
-					]
-					*/
-
-					+SHorizontalBox::Slot()
-					.AutoWidth()
-					.Padding(5.0f)
+				+SHorizontalBox::Slot()
+				.AutoWidth()
+				.Padding(FMargin(5.0f, 0.0f))
+				[
+					SNew(SCheckBox)
+					.Style(FCoreStyle::Get(), "ToggleButtonCheckbox")
+					.IsChecked_Lambda([&]()
+					{
+						return bEnableDemoMode ? ECheckBoxState::Checked : ECheckBoxState::Unchecked;
+					})
+					.OnCheckStateChanged_Lambda([&](const ECheckBoxState NewState)
+					{
+						bEnableDemoMode = (NewState == ECheckBoxState::Checked) ? true : false;
+					})
 					[
-						SNew(SCheckBox)
-						.Style( FCoreStyle::Get(), "ToggleButtonCheckbox" )
-						.IsChecked_Lambda([]()
-						{
-							return bFoldTick == 0 ? ECheckBoxState::Unchecked : ECheckBoxState::Checked;
-						})
-						.OnCheckStateChanged_Lambda([]( const ECheckBoxState NewState )
-						{
-							bFoldTick = (NewState == ECheckBoxState::Checked) ? 1 : 0;
-						})
+						SNew(SBox)
+						.VAlign(VAlign_Center)
+						.HAlign(HAlign_Center)
+						.Padding(FMargin(4.0, 2.0))
 						[
-							SNew(SBox)
-							.VAlign( VAlign_Center )
-							.HAlign( HAlign_Center )
-							[
-								SNew(STextBlock)	
-								.Text( LOCTEXT("ToggleTickFolding", "Fold Tick") )
-							]
+							SNew(STextBlock)
+							.Text(LOCTEXT("EnableDemoMode", "Demo Mode"))
 						]
 					]
-					+ SHorizontalBox::Slot()
-					.AutoWidth()
-					.Padding(5.0f)
-					[
-						SNew(SCheckBox)
-						.Style(FCoreStyle::Get(), "ToggleButtonCheckbox")
-						.IsChecked_Lambda([&]()
+				]
+
+				+SHorizontalBox::Slot()
+				.AutoWidth()
+				.Padding(FMargin(5.0f, 0.0f))
+				[
+					SNew(SCheckBox)
+					.Style( FCoreStyle::Get(), "ToggleButtonCheckbox" )
+					.IsEnabled(SLATE_STATS)
+					.IsChecked_Static([]
+					{
+#if SLATE_STATS
+						return GSlateStatsFlatEnable == 0 ? ECheckBoxState::Unchecked : ECheckBoxState::Checked;
+#else
+						return ECheckBoxState::Unchecked;
+#endif
+					})
+					.OnCheckStateChanged_Lambda( [=]( const ECheckBoxState NewState )
+					{
+#if SLATE_STATS
+						GSlateStatsFlatEnable = (NewState == ECheckBoxState::Checked) ? 1 : 0;
+						if (GSlateStatsFlatEnable)
 						{
-							return bEnableDemoMode ? ECheckBoxState::Checked : ECheckBoxState::Unchecked;
-						})
-						.OnCheckStateChanged_Lambda([&](const ECheckBoxState NewState)
+							TabManager->InvokeTab(WidgetReflectorTabID::SlateStats);
+						}
+						else
 						{
-							bEnableDemoMode = (NewState == ECheckBoxState::Checked) ? true : false;
-						})
+							CloseTab(WidgetReflectorTabID::SlateStats);
+						}
+#endif
+					})
+					.ToolTip
+					(
+						SNew(SToolTip)
 						[
-							SNew(SBox)
-							.VAlign(VAlign_Center)
-							.HAlign(HAlign_Center)
-							[
-								SNew(STextBlock)
-								.Text(LOCTEXT("EnableDemoMode", "Demo Mode"))
-							]
+							SNew(STextBlock)
+							.WrapTextAt(200.0f)
+#if SLATE_STATS
+							.Text( LOCTEXT("ToggleStatsTooltip", "Enables flat stats view.") )
+#else
+							.Text( LOCTEXT("ToggleStatsUnavailableTooltip", "To enable slate stats, compile with SLATE_STATS defined to one (see SlateStats.h).") )
+#endif
 						]
-					]
-#if !(UE_BUILD_SHIPPING || UE_BUILD_TEST)
-					+ SHorizontalBox::Slot()
-					.AutoWidth()
-					.Padding(5.0f)
+					)
 					[
-						SNew(SCheckBox)
-						.Style(FCoreStyle::Get(), "ToggleButtonCheckbox")
-						.IsChecked_Lambda([&]()
-						{
-							return SInvalidationPanel::IsInvalidationDebuggingEnabled() ? ECheckBoxState::Checked : ECheckBoxState::Unchecked;
-						})
-						.OnCheckStateChanged_Lambda([&](const ECheckBoxState NewState)
-						{
-							SInvalidationPanel::EnableInvalidationDebugging(( NewState == ECheckBoxState::Checked ) ? true : false);
-						})
+						SNew(SBox)
+						.VAlign( VAlign_Center )
+						.HAlign( HAlign_Center )
+						.Padding(FMargin(4.0, 2.0))
 						[
-							SNew(SBox)
-							.VAlign(VAlign_Center)
-							.HAlign(HAlign_Center)
-							[
-								SNew(STextBlock)
-								.Text(LOCTEXT("InvalidationDebugging", "Invalidation Debugging"))
-							]
+							SNew(STextBlock)	
+							.Text( LOCTEXT("ToggleStats", "Toggle Stats") )
 						]
 					]
-#endif
+				]
+
+				+SHorizontalBox::Slot()
+				.FillWidth(1.0f)
+				[
+					SNew(SSpacer)
+				]
+
+				+SHorizontalBox::Slot()
+				.AutoWidth()
+				.Padding(FMargin(5.0f, 0.0f))
+				[
+					SNew(SButton)
+					.Text(LOCTEXT("DisplayTextureAtlases", "Display Texture Atlases"))
+					.OnClicked(this, &SWidgetReflector::HandleDisplayTextureAtlases)
+				]
+
+				+SHorizontalBox::Slot()
+				.AutoWidth()
+				.Padding(FMargin(5.0f, 0.0f))
+				[
+					SNew(SButton)
+					.Text(LOCTEXT("DisplayFontAtlases", "Display Font Atlases"))
+					.OnClicked(this, &SWidgetReflector::HandleDisplayFontAtlases)
+				]
+			]
+
+			+SVerticalBox::Slot()
+			[
+				TabManager->RestoreFrom(Layout, nullptr).ToSharedRef()
+			]
+		]
+	];
 
 #if SLATE_STATS
-					+SHorizontalBox::Slot()
-					.AutoWidth()
-					.Padding(5.0f)
+	if (GSlateStatsFlatEnable)
+	{
+		TabManager->InvokeTab(WidgetReflectorTabID::SlateStats);
+	}
+	else
+	{
+		CloseTab(WidgetReflectorTabID::SlateStats);
+	}
+#endif
+}
+
+
+TSharedRef<SDockTab> SWidgetReflector::SpawnWidgetHierarchyTab(const FSpawnTabArgs& Args)
+{
+	return SNew(SDockTab)
+		.Label(LOCTEXT("WidgetHierarchyTab", "Widget Hierarchy"))
+		//.OnCanCloseTab_Lambda([]() { return false; }) // Can't prevent this as it stops the editor from being able to close while the widget reflector is open
+		[
+			SNew(SVerticalBox)
+				
+			+SVerticalBox::Slot()
+			.AutoHeight()
+			.Padding(FMargin(0.0f, 2.0f))
+			[
+				SNew(SHorizontalBox)
+
+				+SHorizontalBox::Slot()
+				.AutoWidth()
+				.Padding(FMargin(5.0f, 0.0f))
+				[
+					// Check box that controls LIVE MODE
+					SNew(SCheckBox)
+					.IsChecked(this, &SWidgetReflector::HandleFocusCheckBoxIsChecked)
+					.OnCheckStateChanged(this, &SWidgetReflector::HandleFocusCheckBoxCheckedStateChanged)
 					[
-						SNew(SCheckBox)
-						.Style( FCoreStyle::Get(), "ToggleButtonCheckbox" )
-						.IsChecked_Static([]
-						{
-							return GSlateStatsFlatEnable == 0 ? ECheckBoxState::Unchecked : ECheckBoxState::Checked;
-						})
-						.OnCheckStateChanged_Lambda( [=]( const ECheckBoxState NewState )
-						{
-							GSlateStatsFlatEnable = (NewState == ECheckBoxState::Checked) ? 1 : 0;
-							StatsToolsBox->SetVisibility(GSlateStatsFlatEnable != 0 ? EVisibility::Visible : EVisibility::Collapsed);
-						})
-						.ToolTip
-						(
-							SNew(SToolTip)
-							[
-								SNew(STextBlock)
-								.WrapTextAt(200.0f)
-								.Text( LOCTEXT("ToggleStatsTooltip", "Enables flat stats view.") )
-							]
-						)
-						[
-							SNew(SBox)
-							.VAlign( VAlign_Center )
-							.HAlign( HAlign_Center )
-							[
-								SNew(STextBlock)	
-								.Text( LOCTEXT("ToggleStats", "Toggle Stats") )
-							]
-						]
+						SNew(STextBlock)
+						.Text(LOCTEXT("ShowFocus", "Show Focus"))
 					]
-#else
-					+SHorizontalBox::Slot()
-					.AutoWidth()
-					.Padding(5.0f)
+				]
+
+				+SHorizontalBox::Slot()
+				.AutoWidth()
+				.Padding(FMargin(5.0f, 0.0f))
+				[
+					// Check box that controls PICKING A WIDGET TO INSPECT
+					SNew(SButton)
+					.OnClicked(this, &SWidgetReflector::HandlePickButtonClicked)
+					.ButtonColorAndOpacity(this, &SWidgetReflector::HandlePickButtonColorAndOpacity)
 					[
-						SNew(SCheckBox)
-						.Style( FCoreStyle::Get(), "ToggleButtonCheckbox" )
-						.ToolTip
-						(
-							SNew(SToolTip)
-							[
-								SNew(STextBlock)
-								.WrapTextAt(200.0f)
-								.Text( LOCTEXT("ToggleStatsUnavailableTooltip", "To enable slate stats, compile with SLATE_STATS defined to one (see SlateStats.h).") )
-							]
-						)
+						SNew(STextBlock)
+						.Text(this, &SWidgetReflector::HandlePickButtonText)
+					]
+				]
+
+				+SHorizontalBox::Slot()
+				.AutoWidth()
+				.Padding(FMargin(5.0f, 0.0f))
+				[
+					// Button that controls taking a snapshot of the current window(s)
+					SNew(SButton)
+					.OnClicked(this, &SWidgetReflector::HandleTakeSnapshotButtonClicked)
+					[
+						SNew(STextBlock)
+						.Text(LOCTEXT("TakeSnapshotButtonText", "Take Snapshot"))
+					]
+				]
+
+#if SLATE_REFLECTOR_HAS_DESKTOP_PLATFORM
+				+SHorizontalBox::Slot()
+				.AutoWidth()
+				.Padding(FMargin(5.0f, 0.0f))
+				[
+					// Button that controls loading a saved snapshot
+					SNew(SButton)
+					.OnClicked(this, &SWidgetReflector::HandleLoadSnapshotButtonClicked)
+					[
+						SNew(STextBlock)
+						.Text(LOCTEXT("LoadSnapshotButtonText", "Load Snapshot"))
+					]
+				]
+#endif // SLATE_REFLECTOR_HAS_DESKTOP_PLATFORM
+			]
+
+			+SVerticalBox::Slot()
+			.FillHeight(1.0f)
+			[
+				SNew(SBorder)
+				.Padding(0)
+				.BorderImage(FCoreStyle::Get().GetBrush("ToolPanel.GroupBorder"))
+				[
+					// The tree view that shows all the info that we capture.
+					SAssignNew(ReflectorTree, SReflectorTree)
+					.ItemHeight(24.0f)
+					.TreeItemsSource(&ReflectorTreeRoot)
+					.OnGenerateRow(this, &SWidgetReflector::HandleReflectorTreeGenerateRow)
+					.OnGetChildren(this, &SWidgetReflector::HandleReflectorTreeGetChildren)
+					.OnSelectionChanged(this, &SWidgetReflector::HandleReflectorTreeSelectionChanged)
+					.HeaderRow
+					(
+						SNew(SHeaderRow)
+
+						+SHeaderRow::Column("WidgetName")
+						.DefaultLabel(LOCTEXT("WidgetName", "Widget Name"))
+						.FillWidth(0.65f)
+
+						+SHeaderRow::Column("ForegroundColor")
+						.FixedWidth(24.0f)
+						.VAlignHeader(VAlign_Center)
+						.HeaderContent()
 						[
-							SNew(SBox)
-							.VAlign( VAlign_Center )
-							.HAlign( HAlign_Center )
-							[
-								SNew(STextBlock)	
-								.Text( LOCTEXT("ToggleStatsUnavailable", "Toggle Stats [see toolip]") )
-							]
+							SNew(STextBlock)
+							.Text(LOCTEXT("ForegroundColor", "FG"))
+							.ToolTipText(LOCTEXT("ForegroundColorToolTip", "Foreground Color"))
 						]
-						.IsEnabled(false)
-					]
-#endif
+
+						+SHeaderRow::Column("Visibility")
+						.DefaultLabel(LOCTEXT("Visibility", "Visibility" ))
+						.FixedWidth(125.0f)
+
+						+SHeaderRow::Column("WidgetInfo")
+						.DefaultLabel(LOCTEXT("WidgetInfo", "Widget Info" ))
+						.FillWidth(0.25f)
+
+						+SHeaderRow::Column("Address")
+						.DefaultLabel( LOCTEXT("Address", "Address") )
+						.FixedWidth(140.0f)
+					)
 				]
+			]
+		];
+}
+
+
 #if SLATE_STATS
-				+ SVerticalBox::Slot()
-				.AutoHeight()
+
+TSharedRef<SDockTab> SWidgetReflector::SpawnSlateStatsTab(const FSpawnTabArgs& Args)
+{
+	auto OnTabClosed = [](TSharedRef<SDockTab>)
+	{
+		// Tab closed - disable stats
+		GSlateStatsFlatEnable = 0;
+	};
+
+	return SNew(SDockTab)
+		.Label(LOCTEXT("SlateStatsTab", "Slate Stats"))
+		.OnTabClosed_Lambda(OnTabClosed)
+		[
+			SNew(SVerticalBox)
+				
+			+SVerticalBox::Slot()
+			.AutoHeight()
+			.Padding(FMargin(0.0f, 2.0f))
+			[
+				SNew(SHorizontalBox)
+
+				+SHorizontalBox::Slot()
+				.AutoWidth()
+				.Padding(FMargin(5.0f, 0.0f))
 				[
-					SAssignNew(StatsToolsBox, SHorizontalBox)
-					.Visibility(EVisibility::Collapsed)
-					+SHorizontalBox::Slot()
-					.AutoWidth()
-					.Padding(5.0f)
+					SNew(SCheckBox)
+					.Style( FCoreStyle::Get(), "ToggleButtonCheckbox" )
+					.IsChecked_Static([]
+					{
+						return GSlateStatsFlatLogOutput == 0 ? ECheckBoxState::Unchecked : ECheckBoxState::Checked;
+					})
+					.OnCheckStateChanged_Static([]( const ECheckBoxState NewState )
+					{
+						GSlateStatsFlatLogOutput = (NewState == ECheckBoxState::Checked) ? 1 : 0;
+					})
+					.ToolTip
+					(
+						SNew(SToolTip)
+						[
+							SNew(STextBlock)
+							.WrapTextAt(200.0f)
+							.Text( LOCTEXT("LogStatsTooltip", "Enables outputting stats to the log at the given interval.") )
+						]
+					)
 					[
-						SNew(SCheckBox)
-						.Style( FCoreStyle::Get(), "ToggleButtonCheckbox" )
-						.IsChecked_Static([]
-						{
-							return GSlateStatsFlatLogOutput == 0 ? ECheckBoxState::Unchecked : ECheckBoxState::Checked;
-						})
-						.OnCheckStateChanged_Static([]( const ECheckBoxState NewState )
-						{
-							GSlateStatsFlatLogOutput = (NewState == ECheckBoxState::Checked) ? 1 : 0;
-						})
-						.ToolTip
-						(
-							SNew(SToolTip)
-							[
-								SNew(STextBlock)
-								.WrapTextAt(200.0f)
-								.Text( LOCTEXT("LogStatsTooltip", "Enables outputting stats to the log at the given interval.") )
-							]
-						)
+						SNew(SBox)
+						.VAlign( VAlign_Center )
+						.HAlign( HAlign_Center )
 						[
-							SNew(SBox)
-							.VAlign( VAlign_Center )
-							.HAlign( HAlign_Center )
-							[
-								SNew(STextBlock)	
-								.Text( LOCTEXT("ToggleLogStats", "Log Stats") )
-							]
+							SNew(STextBlock)	
+							.Text( LOCTEXT("ToggleLogStats", "Log Stats") )
 						]
 					]
-					+SHorizontalBox::Slot()
-					.AutoWidth()
-					.Padding(5.0f)
-					[
-						SNew(SButton)
-						.OnClicked_Static([]
-						{
-							GSlateStatsHierarchyTrigger = 1;
-							return FReply::Handled();
-						})
-						.ToolTip
-						(
-							SNew(SToolTip)
-							[
-								SNew(STextBlock)
-								.WrapTextAt(200.0f)
-								.Text( LOCTEXT("CaptureStatsHierarchyTooltip", "When clicked, the next rendered frame will capture hierarchical stats and save them to file in the Saved/ folder with the following name: SlateHierarchyStats-<timestamp>.csv") )
-							]
-						)
+				]
+
+				+SHorizontalBox::Slot()
+				.AutoWidth()
+				.Padding(FMargin(5.0f, 0.0f))
+				[
+					SNew(SButton)
+					.OnClicked_Static([]
+					{
+						GSlateStatsHierarchyTrigger = 1;
+						return FReply::Handled();
+					})
+					.ToolTip
+					(
+						SNew(SToolTip)
 						[
 							SNew(STextBlock)
-								.Text(LOCTEXT("CaptureHierarchy", "Capture Hierarchy"))
+							.WrapTextAt(200.0f)
+							.Text( LOCTEXT("CaptureStatsHierarchyTooltip", "When clicked, the next rendered frame will capture hierarchical stats and save them to file in the Saved/ folder with the following name: SlateHierarchyStats-<timestamp>.csv") )
 						]
-					]
-					+ SHorizontalBox::Slot()
-					.AutoWidth()
-					.VAlign(VAlign_Center)
+					)
 					[
 						SNew(STextBlock)
-						.ToolTip
-						(
-							SNew(SToolTip)
-							[
-								SNew(STextBlock)
-								.WrapTextAt(200.0f)
-								.Text( LOCTEXT("StatsSamplingIntervalTooltip", "the interval (in seconds) to integrate stats before updating the averages.") )
-							]
-						)
-						.Text(LOCTEXT("StatsSampleWindow", "Sampling Interval: "))
+							.Text(LOCTEXT("CaptureHierarchy", "Capture Hierarchy"))
 					]
-					+ SHorizontalBox::Slot()
+				]
+
+				+SHorizontalBox::Slot()
+				.AutoWidth()
+				.VAlign(VAlign_Center)
+				[
+					SNew(STextBlock)
+					.ToolTip
+					(
+						SNew(SToolTip)
+						[
+							SNew(STextBlock)
+							.WrapTextAt(200.0f)
+							.Text( LOCTEXT("StatsSamplingIntervalTooltip", "the interval (in seconds) to integrate stats before updating the averages.") )
+						]
+					)
+					.Text(LOCTEXT("StatsSampleWindow", "Sampling Interval: "))
+				]
+
+				+SHorizontalBox::Slot()
+				.AutoWidth()
+				[
+					SNew(SBox)
+					.MinDesiredWidth(100)
+					.MaxDesiredWidth(250)
 					[
 						SNew(SSpinBox<float>)
 						.ToolTip
@@ -628,76 +834,117 @@ void SWidgetReflector::Construct( const FArguments& InArgs )
 						})
 					]
 				]
-#endif
-				+ SVerticalBox::Slot()
-					.FillHeight(1.0f)
-					[
-						// The tree view that shows all the info that we capture.
-						SAssignNew(ReflectorTree, SReflectorTree)
-						.ItemHeight(24.0f)
-						.TreeItemsSource(&ReflectorTreeRoot)
-						.OnGenerateRow(this, &SWidgetReflector::HandleReflectorTreeGenerateRow)
-						.OnGetChildren(this, &SWidgetReflector::HandleReflectorTreeGetChildren)
-						.OnSelectionChanged(this, &SWidgetReflector::HandleReflectorTreeSelectionChanged)
-						.HeaderRow
-						(
-							SNew(SHeaderRow)
+			]
 
-							+ SHeaderRow::Column("WidgetName")
-							.DefaultLabel(LOCTEXT("WidgetName", "Widget Name"))
-							.FillWidth(0.65f)
+#if WITH_EVENT_LOGGING
+			+SVerticalBox::Slot()
+			.FillHeight(1.0f)
+			[
+				SNew(SBorder)
+				.Padding(0)
+				.BorderImage(FCoreStyle::Get().GetBrush("ToolPanel.GroupBorder"))
+				[
+					SAssignNew(EventListView, SListView<TSharedRef<FLoggedEvent>>)
+					.ListItemsSource( &LoggedEvents )
+					.OnGenerateRow(this, &SWidgetReflector::GenerateEventLogRow)
+				]
+			]
+#endif //WITH_EVENT_LOGGING
 
-							+ SHeaderRow::Column("ForegroundColor")
-							.FixedWidth(24.0f)
-							.VAlignHeader(VAlign_Center)
-							.HeaderContent()
-							[
-								SNew(STextBlock)
-								.Text(LOCTEXT("ForegroundColor", "FG"))
-								.ToolTipText(LOCTEXT("ForegroundColorToolTip", "Foreground Color"))
-							]
+			+SVerticalBox::Slot()
+			.FillHeight(1.0f)
+			[
+				MakeStatViewer()
+			]
+		];
+}
 
-							+ SHeaderRow::Column("Visibility")
-							.DefaultLabel(LOCTEXT("Visibility", "Visibility" ))
-							.FixedWidth(125.0f)
+#endif // #if SLATE_STATS
 
-							+ SHeaderRow::Column("WidgetInfo")
-							.DefaultLabel(LOCTEXT("WidgetInfo", "Widget Info" ))
-							.FillWidth(0.25f)
 
-							+ SHeaderRow::Column("Address")
-							.DefaultLabel( LOCTEXT("Address", "Address") )
-							.FixedWidth(140.0f)
-						)
-					]
+TSharedRef<SDockTab> SWidgetReflector::SpawnSnapshotWidgetPicker(const FSpawnTabArgs& Args)
+{
+	auto OnTabClosed = [this](TSharedRef<SDockTab>)
+	{
+		// Tab closed - leave snapshot mode
+		SetUIMode(EWidgetReflectorUIMode::Live);
+	};
+
+	auto OnWidgetPathPicked = [this](const TArray<TSharedRef<FWidgetReflectorNodeBase>>& PickedWidgetPath)
+	{
+		VisualizeAsTree(PickedWidgetPath);
+	};
+
+	return SNew(SDockTab)
+		.Label(LOCTEXT("SnapshotWidgetPickerTab", "Snapshot Widget Picker"))
+		.OnTabClosed_Lambda(OnTabClosed)
+		[
+			SAssignNew(WidgetSnapshotVisualizer, SWidgetSnapshotVisualizer)
+			.SnapshotData(&SnapshotData)
+			.OnWidgetPathPicked_Lambda(OnWidgetPathPicked)
+		];
+}
+
+
+void SWidgetReflector::OnTabSpawned(const FName& TabIdentifier, const TSharedRef<SDockTab>& SpawnedTab)
+{
+	TWeakPtr<SDockTab>* const ExistingTab = SpawnedTabs.Find(TabIdentifier);
+	if (!ExistingTab)
+	{
+		SpawnedTabs.Add(TabIdentifier, SpawnedTab);
+	}
+	else
+	{
+		check(!ExistingTab->IsValid());
+		*ExistingTab = SpawnedTab;
+	}
+}
 
-				#if WITH_EVENT_LOGGING
-				+ SVerticalBox::Slot()
-					.FillHeight(1.0f)
-					[
-						SAssignNew(EventListView, SListView<TSharedRef<FLoggedEvent>>)
-							.ListItemsSource( &LoggedEvents )
-							.OnGenerateRow(this, &SWidgetReflector::GenerateEventLogRow)
-					]
-				#endif //WITH_EVENT_LOGGING
-#if SLATE_STATS			
-				+ SVerticalBox::Slot()
-					.AutoHeight()
-					[
-						MakeStatViewer()
-					]
-#endif
-			]
-	];
 
-#if SLATE_STATS			
-	if (StatsToolsBox.IsValid())
+void SWidgetReflector::CloseTab(const FName& TabIdentifier)
+{
+	TWeakPtr<SDockTab>* const ExistingTab = SpawnedTabs.Find(TabIdentifier);
+	if (ExistingTab)
 	{
-		StatsToolsBox->SetVisibility(GSlateStatsFlatEnable != 0 ? EVisibility::Visible : EVisibility::Collapsed);
+		TSharedPtr<SDockTab> ExistingTabPin = ExistingTab->Pin();
+		if (ExistingTabPin.IsValid())
+		{
+			ExistingTabPin->RequestCloseTab();
+		}
 	}
-#endif
 }
 
+
+void SWidgetReflector::SetUIMode(const EWidgetReflectorUIMode InNewMode)
+{
+	if (CurrentUIMode != InNewMode)
+	{
+		CurrentUIMode = InNewMode;
+
+		SelectedNodes.Reset();
+		ReflectorTreeRoot.Reset();
+		PickedPath.Reset();
+		ReflectorTree->RequestTreeRefresh();
+
+		if (CurrentUIMode == EWidgetReflectorUIMode::Snapshot)
+		{
+			TabManager->InvokeTab(WidgetReflectorTabID::SnapshotWidgetPicker);
+		}
+		else
+		{
+			SnapshotData.ClearSnapshot();
+
+			if (WidgetSnapshotVisualizer.IsValid())
+			{
+				WidgetSnapshotVisualizer->SnapshotDataUpdated();
+			}
+
+			CloseTab(WidgetReflectorTabID::SnapshotWidgetPicker);
+		}
+	}
+}
+
+
 /* SCompoundWidget overrides
  *****************************************************************************/
 
@@ -824,7 +1071,7 @@ int32 SWidgetReflector::VisualizeCursorAndKeys(FSlateWindowElementList& OutDrawE
 /* SWidgetReflector implementation
  *****************************************************************************/
 
-TSharedRef<SToolTip> SWidgetReflector::GenerateToolTipForReflectorNode( UWidgetReflectorNodeBase* InReflectorNode )
+TSharedRef<SToolTip> SWidgetReflector::GenerateToolTipForReflectorNode( TSharedRef<FWidgetReflectorNodeBase> InReflectorNode )
 {
 	return SNew(SToolTip)
 		[
@@ -834,14 +1081,14 @@ TSharedRef<SToolTip> SWidgetReflector::GenerateToolTipForReflectorNode( UWidgetR
 }
 
 
-void SWidgetReflector::VisualizeAsTree( const TArray<UWidgetReflectorNodeBase*>& WidgetPathToVisualize )
+void SWidgetReflector::VisualizeAsTree( const TArray<TSharedRef<FWidgetReflectorNodeBase>>& WidgetPathToVisualize )
 {
 	const FLinearColor TopmostWidgetColor(1.0f, 0.0f, 0.0f);
 	const FLinearColor LeafmostWidgetColor(0.0f, 1.0f, 0.0f);
 
 	for (int32 WidgetIndex = 0; WidgetIndex<WidgetPathToVisualize.Num(); ++WidgetIndex)
 	{
-		UWidgetReflectorNodeBase* CurWidget = WidgetPathToVisualize[WidgetIndex];
+		const auto& CurWidget = WidgetPathToVisualize[WidgetIndex];
 
 		// Tint the item based on depth in picked path
 		const float ColorFactor = static_cast<float>(WidgetIndex)/WidgetPathToVisualize.Num();
@@ -890,11 +1137,11 @@ int32 SWidgetReflector::VisualizePickAsRectangles( const FWidgetPath& InWidgetsT
 }
 
 
-int32 SWidgetReflector::VisualizeSelectedNodesAsRectangles( const TArray<UWidgetReflectorNodeBase*>& InNodesToDraw, const TSharedRef<SWindow>& VisualizeInWindow, FSlateWindowElementList& OutDrawElements, int32 LayerId )
+int32 SWidgetReflector::VisualizeSelectedNodesAsRectangles( const TArray<TSharedRef<FWidgetReflectorNodeBase>>& InNodesToDraw, const TSharedRef<SWindow>& VisualizeInWindow, FSlateWindowElementList& OutDrawElements, int32 LayerId )
 {
 	for (int32 NodeIndex = 0; NodeIndex < InNodesToDraw.Num(); ++NodeIndex)
 	{
-		const UWidgetReflectorNodeBase* NodeToDraw = InNodesToDraw[NodeIndex];
+		const auto& NodeToDraw = InNodesToDraw[NodeIndex];
 		const FLinearColor Tint(0.0f, 1.0f, 0.0f);
 
 		// The FGeometry we get is from a WidgetPath, so it's rooted in desktop space.
@@ -902,7 +1149,7 @@ int32 SWidgetReflector::VisualizeSelectedNodesAsRectangles( const TArray<UWidget
 		// and get us back into Window Space.
 		// This is nonstandard so we have to go through some hoops and a specially exposed method 
 		// in FPaintGeometry to allow appending layout transforms.
-		FPaintGeometry WindowSpaceGeometry = NodeToDraw->GetGeometry().ToPaintGeometry();
+		FPaintGeometry WindowSpaceGeometry(NodeToDraw->GetAccumulatedLayoutTransform(), NodeToDraw->GetAccumulatedRenderTransform(), NodeToDraw->GetLocalSize());
 		WindowSpaceGeometry.AppendTransform(TransformCast<FSlateLayoutTransform>(Inverse(VisualizeInWindow->GetPositionInScreen())));
 
 		FSlateDrawElement::MakeBox(
@@ -961,33 +1208,70 @@ FString SWidgetReflector::HandleFrameRateText() const
 
 FText SWidgetReflector::HandlePickButtonText() const
 {
-	static const FText NotPicking = LOCTEXT("PickWidget", "Pick Widget");
+	static const FText NotPicking = LOCTEXT("PickLiveWidget", "Pick Live Widget");
 	static const FText Picking = LOCTEXT("PickingWidget", "Picking (Esc to Stop)");
 
 	return bIsPicking ? Picking : NotPicking;
 }
 
 
-FReply SWidgetReflector::HandleSnapshotButtonClicked()
+FReply SWidgetReflector::HandleTakeSnapshotButtonClicked()
 {
-	// Clean out the current tree and replace it with the snapshotted state of any window(s) that are currently open
-	ReflectorTreeRoot.Empty();
+	SetUIMode(EWidgetReflectorUIMode::Snapshot);
+
+	// Take a snapshot of any window(s) that are currently open
+	SnapshotData.TakeSnapshot();
+
+	// Rebuild the reflector tree from the snapshot data
+	ReflectorTreeRoot = SnapshotData.GetWindowsRef();
+	ReflectorTree->RequestTreeRefresh();
+
+	WidgetSnapshotVisualizer->SnapshotDataUpdated();
+
+	return FReply::Handled();
+}
+
+
+#if SLATE_REFLECTOR_HAS_DESKTOP_PLATFORM
 
-	TArray<TSharedRef<SWindow>> VisibleWindows;
-	FSlateApplication::Get().GetAllVisibleWindowsOrdered(VisibleWindows);
+FReply SWidgetReflector::HandleLoadSnapshotButtonClicked()
+{
+	IDesktopPlatform* DesktopPlatform = FDesktopPlatformModule::Get();
 
-	for (const auto& VisibleWindow : VisibleWindows)
+	if (DesktopPlatform)
 	{
-		ReflectorTreeRoot.Add(FWidgetReflectorNodeUtils::NewSnapshotNodeTreeFrom(FArrangedWidget(VisibleWindow, VisibleWindow->GetWindowGeometryInScreen())));
-	}
+		TSharedPtr<SWindow> ParentWindow = FSlateApplication::Get().FindWidgetWindow(SharedThis(this));
+
+		TArray<FString> OpenFilenames;
+		const bool bOpened = DesktopPlatform->OpenFileDialog(
+			(ParentWindow.IsValid()) ? ParentWindow->GetNativeWindow()->GetOSWindowHandle() : nullptr,
+			LOCTEXT("LoadSnapshotDialogTitle", "Load Widget Snapshot").ToString(),
+			FPaths::GameAgnosticSavedDir(),
+			TEXT(""),
+			TEXT("Slate Widget Snapshot (*.widgetsnapshot)|*.widgetsnapshot"),
+			EFileDialogFlags::None,
+			OpenFilenames
+			);
+
+		if (bOpened && SnapshotData.LoadSnapshotFromFile(OpenFilenames[0]))
+		{
+			SetUIMode(EWidgetReflectorUIMode::Snapshot);
 
-	ReflectorTree->RequestTreeRefresh();
+			// Rebuild the reflector tree from the snapshot data
+			ReflectorTreeRoot = SnapshotData.GetWindowsRef();
+			ReflectorTree->RequestTreeRefresh();
+
+			WidgetSnapshotVisualizer->SnapshotDataUpdated();
+		}
+	}
 
 	return FReply::Handled();
 }
 
+#endif // SLATE_REFLECTOR_HAS_DESKTOP_PLATFORM
+
 
-TSharedRef<ITableRow> SWidgetReflector::HandleReflectorTreeGenerateRow( UWidgetReflectorNodeBase* InReflectorNode, const TSharedRef<STableViewBase>& OwnerTable )
+TSharedRef<ITableRow> SWidgetReflector::HandleReflectorTreeGenerateRow( TSharedRef<FWidgetReflectorNodeBase> InReflectorNode, const TSharedRef<STableViewBase>& OwnerTable )
 {
 	return SNew(SReflectorTreeWidgetItem, OwnerTable)
 		.WidgetInfoToVisualize(InReflectorNode)
@@ -997,9 +1281,9 @@ TSharedRef<ITableRow> SWidgetReflector::HandleReflectorTreeGenerateRow( UWidgetR
 }
 
 
-void SWidgetReflector::HandleReflectorTreeGetChildren(UWidgetReflectorNodeBase* InWidgetGeometry, TArray<UWidgetReflectorNodeBase*>& OutChildren)
+void SWidgetReflector::HandleReflectorTreeGetChildren(TSharedRef<FWidgetReflectorNodeBase> InReflectorNode, TArray<TSharedRef<FWidgetReflectorNodeBase>>& OutChildren)
 {
-	OutChildren = InWidgetGeometry->GetChildNodes();
+	OutChildren = InReflectorNode->GetChildNodes();
 }
 
 
@@ -1059,7 +1343,8 @@ TSharedRef<SWidget> SWidgetReflector::MakeStatViewer()
 
 	return
 		SAssignNew(StatsBorder, SBorder)
-		.BorderImage( FCoreStyle::Get().GetBrush("NoBrush") )
+		.Padding(0)
+		.BorderImage(FCoreStyle::Get().GetBrush("ToolPanel.GroupBorder"))
 		.Visibility_Lambda([]
 		{
 			return GSlateStatsFlatEnable > 0 ? EVisibility::Visible : EVisibility::Collapsed;
diff --git a/Engine/Source/Developer/SlateReflector/Private/Widgets/SWidgetReflector.h b/Engine/Source/Developer/SlateReflector/Private/Widgets/SWidgetReflector.h
--- a/Engine/Source/Developer/SlateReflector/Private/Widgets/SWidgetReflector.h
+++ b/Engine/Source/Developer/SlateReflector/Private/Widgets/SWidgetReflector.h
@@ -14,8 +14,11 @@ class SLATEREFLECTOR_API SWidgetReflector
 {
 public:
 
-	SLATE_USER_ARGS(SWidgetReflector) 
+	SLATE_USER_ARGS(SWidgetReflector)
 	{ }
+		
+		SLATE_ARGUMENT(TSharedPtr<SDockTab>, ParentTab)
+
 	SLATE_END_ARGS()
 
 	virtual void Construct( const FArguments& InArgs ) = 0;
diff --git a/Engine/Source/Developer/SlateReflector/Private/Widgets/SWidgetReflectorToolTipWidget.h b/Engine/Source/Developer/SlateReflector/Private/Widgets/SWidgetReflectorToolTipWidget.h
--- a/Engine/Source/Developer/SlateReflector/Private/Widgets/SWidgetReflectorToolTipWidget.h
+++ b/Engine/Source/Developer/SlateReflector/Private/Widgets/SWidgetReflectorToolTipWidget.h
@@ -15,7 +15,7 @@ public:
 		: _WidgetInfoToVisualize()
 	{ }
 
-		SLATE_ARGUMENT(UWidgetReflectorNodeBase*, WidgetInfoToVisualize)
+		SLATE_ARGUMENT(TSharedPtr<FWidgetReflectorNodeBase>, WidgetInfoToVisualize)
 
 	SLATE_END_ARGS()
 
@@ -25,7 +25,10 @@ public:
 	void Construct( const FArguments& InArgs )
 	{
 		this->WidgetInfo = InArgs._WidgetInfoToVisualize;
-		check(WidgetInfo);
+		check(WidgetInfo.IsValid());
+
+		const FGeometry Geom = FGeometry().MakeChild(WidgetInfo->GetLocalSize(), WidgetInfo->GetAccumulatedLayoutTransform(), WidgetInfo->GetAccumulatedRenderTransform(), FVector2D::ZeroVector);
+		SizeInfoText = FText::FromString(Geom.ToString());
 
 		ChildSlot
 		[
@@ -103,12 +106,12 @@ private:
 
 	FText GetWidgetActualSize() const
 	{
-		return FText::FromString(WidgetInfo->GetGeometry().Size.ToString());
+		return FText::FromString(WidgetInfo->GetLocalSize().ToString());
 	}
 
 	FText GetSizeInfo() const
 	{
-		return FText::FromString(WidgetInfo->GetGeometry().ToString());
+		return SizeInfoText;
 	}
 
 	FText GetEnabled() const
@@ -121,7 +124,10 @@ private:
 private:
 
 	/** The info about the widget that we are visualizing. */
-	UWidgetReflectorNodeBase* WidgetInfo;
+	TSharedPtr<FWidgetReflectorNodeBase> WidgetInfo;
+
+	/** The size info text */
+	FText SizeInfoText;
 };
 
 
diff --git a/Engine/Source/Developer/SlateReflector/Private/Widgets/SWidgetReflectorTreeWidgetItem.h b/Engine/Source/Developer/SlateReflector/Private/Widgets/SWidgetReflectorTreeWidgetItem.h
--- a/Engine/Source/Developer/SlateReflector/Private/Widgets/SWidgetReflectorTreeWidgetItem.h
+++ b/Engine/Source/Developer/SlateReflector/Private/Widgets/SWidgetReflectorTreeWidgetItem.h
@@ -7,7 +7,7 @@
  * Widget that visualizes the contents of a FReflectorNode.
  */
 class SLATEREFLECTOR_API SReflectorTreeWidgetItem
-	: public SMultiColumnTableRow<UWidgetReflectorNodeBase*>
+	: public SMultiColumnTableRow<TSharedRef<FWidgetReflectorNodeBase>>
 {
 public:
 
@@ -17,7 +17,7 @@ public:
 		, _AssetAccessor()
 	{ }
 
-		SLATE_ARGUMENT(UWidgetReflectorNodeBase*, WidgetInfoToVisualize)
+		SLATE_ARGUMENT(TSharedPtr<FWidgetReflectorNodeBase>, WidgetInfoToVisualize)
 		SLATE_ARGUMENT(FAccessSourceCode, SourceCodeAccessor)
 		SLATE_ARGUMENT(FAccessAsset, AssetAccessor)
 
@@ -36,14 +36,14 @@ public:
 		this->OnAccessSourceCode = InArgs._SourceCodeAccessor;
 		this->OnAccessAsset = InArgs._AssetAccessor;
 
-		check(WidgetInfo);
+		check(WidgetInfo.IsValid());
 		CachedWidgetType = WidgetInfo->GetWidgetType();
 		CachedWidgetVisibility = WidgetInfo->GetWidgetVisibilityText();
 		CachedReadableLocation = WidgetInfo->GetWidgetReadableLocation();
 		CachedWidgetFile = WidgetInfo->GetWidgetFile();
 		CachedWidgetLineNumber = WidgetInfo->GetWidgetLineNumber();
 
-		SMultiColumnTableRow< UWidgetReflectorNodeBase* >::Construct( SMultiColumnTableRow< UWidgetReflectorNodeBase* >::FArguments().Padding(1), InOwnerTableView );
+		SMultiColumnTableRow< TSharedRef<FWidgetReflectorNodeBase> >::Construct( SMultiColumnTableRow< TSharedRef<FWidgetReflectorNodeBase> >::FArguments().Padding(1), InOwnerTableView );
 	}
 
 public:
@@ -96,7 +96,7 @@ protected:
 private:
 
 	/** The info about the widget that we are visualizing. */
-	UWidgetReflectorNodeBase* WidgetInfo;
+	TSharedPtr<FWidgetReflectorNodeBase> WidgetInfo;
 
 	FText CachedWidgetType;
 	FText CachedWidgetVisibility;
diff --git a/Engine/Source/Developer/SlateReflector/Private/Widgets/SWidgetSnapshotVisualizer.cpp b/Engine/Source/Developer/SlateReflector/Private/Widgets/SWidgetSnapshotVisualizer.cpp
new file mode 100644
--- /dev/null
+++ b/Engine/Source/Developer/SlateReflector/Private/Widgets/SWidgetSnapshotVisualizer.cpp
@@ -0,0 +1,891 @@
+// Copyright 1998-2015 Epic Games, Inc. All Rights Reserved.
+
+#include "SlateReflectorPrivatePCH.h"
+#include "SWidgetSnapshotVisualizer.h"
+#include "WidgetReflectorNode.h"
+#include "ScrollyZoomy.h"
+#include "Base64.h"
+
+#if SLATE_REFLECTOR_HAS_DESKTOP_PLATFORM
+#include "DesktopPlatformModule.h"
+#endif // SLATE_REFLECTOR_HAS_DESKTOP_PLATFORM
+
+#define LOCTEXT_NAMESPACE "WidgetSnapshotVisualizer"
+
+class SScrollableSnapshotImage : public IScrollableZoomable, public SPanel
+{
+public:
+	class FScrollableSnapshotImageSlot : public TSupportsOneChildMixin<FScrollableSnapshotImageSlot>
+	{
+	public:
+		FScrollableSnapshotImageSlot()
+			: TSupportsOneChildMixin<FScrollableSnapshotImageSlot>()
+		{
+		}
+	};
+
+	SLATE_BEGIN_ARGS(SScrollableSnapshotImage)
+		: _SnapshotData(nullptr)
+		{
+			_Visibility = EVisibility::Visible;
+		}
+		
+		SLATE_ARGUMENT(const FWidgetSnapshotData*, SnapshotData)
+
+		SLATE_EVENT(SWidgetSnapshotVisualizer::FOnWidgetPathPicked, OnWidgetPathPicked);
+
+	SLATE_END_ARGS()
+
+	SScrollableSnapshotImage()
+		: PhysicalOffset(ForceInitToZero)
+		, CachedSize(ForceInitToZero)
+		, ChildSlot()
+		, ScrollyZoomy(false)
+		, SnapshotDataPtr(nullptr)
+		, bIsPicking(false)
+	{
+	}
+
+	void Construct(const FArguments& InArgs)
+	{
+		SnapshotDataPtr = InArgs._SnapshotData;
+		check(SnapshotDataPtr);
+
+		SelectedWindowIndex = INDEX_NONE;
+
+		OnWidgetPathPicked = InArgs._OnWidgetPathPicked;
+
+		ChildSlot
+		[
+			SNew(SImage)
+			.Image(this, &SScrollableSnapshotImage::GetSelectedWindowTextureBrush)
+		];
+	}
+
+	void SetSelectedWindowIndex(const int32 InIndex)
+	{
+		SelectedWindowIndex = InIndex;
+		PickedWidgets.Reset();
+		PhysicalOffset = FVector2D::ZeroVector;
+	}
+
+	int32 GetSelectedWindowIndex() const
+	{
+		return SelectedWindowIndex;
+	}
+
+	const FSlateBrush* GetSelectedWindowTextureBrush() const
+	{
+		return SnapshotDataPtr->GetBrush(SelectedWindowIndex);
+	}
+
+	void SetIsPicking(const bool InIsPicking)
+	{
+		bIsPicking = InIsPicking;
+	}
+
+	bool GetIsPicking() const
+	{
+		return bIsPicking;
+	}
+
+	virtual void OnArrangeChildren(const FGeometry& AllottedGeometry, FArrangedChildren& ArrangedChildren) const override
+	{
+		CachedSize = AllottedGeometry.GetLocalSize();
+
+		const TSharedRef<SWidget>& ChildWidget = ChildSlot.GetWidget();
+		if (ChildWidget->GetVisibility() != EVisibility::Collapsed)
+		{
+			const FVector2D& WidgetDesiredSize = ChildWidget->GetDesiredSize();
+
+			// Clamp the pan offset based on our current geometry
+			SScrollableSnapshotImage* const NonConstThis = const_cast<SScrollableSnapshotImage*>(this);
+			NonConstThis->ClampViewOffset(WidgetDesiredSize, CachedSize);
+
+			ArrangedChildren.AddWidget(AllottedGeometry.MakeChild(ChildWidget, PhysicalOffset, WidgetDesiredSize));
+		}
+	}
+
+	virtual FVector2D ComputeDesiredSize(float) const override
+	{
+		FVector2D ThisDesiredSize = FVector2D::ZeroVector;
+
+		const TSharedRef<SWidget>& ChildWidget = ChildSlot.GetWidget();
+		if (ChildWidget->GetVisibility() != EVisibility::Collapsed)
+		{
+			ThisDesiredSize = ChildWidget->GetDesiredSize();
+		}
+
+		return ThisDesiredSize;
+	}
+
+	virtual FChildren* GetChildren() override
+	{
+		return &ChildSlot;
+	}
+
+	virtual void Tick(const FGeometry& AllottedGeometry, const double InCurrentTime, const float InDeltaTime) override
+	{
+		ScrollyZoomy.Tick(InDeltaTime, *this);
+	}
+
+	virtual FReply OnMouseButtonDown(const FGeometry& MyGeometry, const FPointerEvent& MouseEvent) override
+	{
+		return ScrollyZoomy.OnMouseButtonDown(MouseEvent);
+	}
+
+	virtual FReply OnMouseButtonUp(const FGeometry& MyGeometry, const FPointerEvent& MouseEvent) override
+	{
+		return ScrollyZoomy.OnMouseButtonUp(AsShared(), MyGeometry, MouseEvent);
+	}
+		
+	virtual FReply OnMouseMove(const FGeometry& MyGeometry, const FPointerEvent& MouseEvent) override
+	{
+		struct FWidgetPicker
+		{
+			static bool FindWidgetsUnderPoint(const FVector2D& InHitTestPoint, const FVector2D& InWindowPosition, const TSharedRef<FWidgetReflectorNodeBase>& InWidget, TArray<TSharedRef<FWidgetReflectorNodeBase>>& OutWidgets)
+			{
+				const bool bNeedsHitTesting = InWidget->GetHitTestInfo().IsHitTestVisible || InWidget->GetHitTestInfo().AreChildrenHitTestVisible;
+				if (bNeedsHitTesting)
+				{
+					const FSlateRect HitTestRect = FSlateRect::FromPointAndExtent(
+						InWidget->GetAccumulatedLayoutTransform().GetTranslation() - InWindowPosition, 
+						TransformPoint(InWidget->GetAccumulatedLayoutTransform().GetScale(), InWidget->GetLocalSize())
+						);
+
+					if (HitTestRect.ContainsPoint(InHitTestPoint))
+					{
+						OutWidgets.Add(InWidget);
+
+						if (InWidget->GetHitTestInfo().AreChildrenHitTestVisible)
+						{
+							for (const auto& ChildWidget : InWidget->GetChildNodes())
+							{
+								if (FindWidgetsUnderPoint(InHitTestPoint, InWindowPosition, ChildWidget, OutWidgets))
+								{
+									return true;
+								}
+							}
+						}
+
+						return InWidget->GetHitTestInfo().IsHitTestVisible;
+					}
+				}
+
+				return false;
+			}
+		};
+
+		if (bIsPicking)
+		{
+			// We need to pick in the snapshot window space, so convert the mouse co-ordinates to be relative to our top-left position
+			const FVector2D& ScreenMousePos = MouseEvent.GetScreenSpacePosition();
+			const FVector2D LocalMousePos = MyGeometry.AbsoluteToLocal(ScreenMousePos);
+			const FVector2D ScrolledPos = LocalMousePos - PhysicalOffset;
+
+			PickedWidgets.Reset();
+
+			TSharedPtr<FWidgetReflectorNodeBase> Window = SnapshotDataPtr->GetWindow(SelectedWindowIndex);
+			if (Window.IsValid())
+			{
+				FWidgetPicker::FindWidgetsUnderPoint(
+					ScrolledPos, 
+					Window->GetAccumulatedLayoutTransform().GetTranslation(), 
+					Window.ToSharedRef(), 
+					PickedWidgets
+					);
+			}
+
+			if (PickedWidgets.Num() > 0)
+			{
+				OnWidgetPathPicked.ExecuteIfBound(PickedWidgets);
+			}
+		}
+
+		return ScrollyZoomy.OnMouseMove(AsShared(), *this, MyGeometry, MouseEvent);
+	}
+
+	virtual void OnMouseLeave(const FPointerEvent& MouseEvent) override
+	{
+		ScrollyZoomy.OnMouseLeave(AsShared(), MouseEvent);
+	}
+		
+	virtual FReply OnMouseWheel(const FGeometry& MyGeometry, const FPointerEvent& MouseEvent) override
+	{
+		return ScrollyZoomy.OnMouseWheel(MouseEvent, *this);
+	}
+
+	virtual FCursorReply OnCursorQuery(const FGeometry& MyGeometry, const FPointerEvent& CursorEvent) const override
+	{
+		FCursorReply Reply = ScrollyZoomy.OnCursorQuery();
+		
+		if (!Reply.IsEventHandled() && !bIsPicking)
+		{
+			Reply = FCursorReply::Cursor(EMouseCursor::GrabHand);
+		}
+
+		return Reply;
+	}
+
+	virtual int32 OnPaint(const FPaintArgs& Args, const FGeometry& AllottedGeometry, const FSlateRect& MyClippingRect, FSlateWindowElementList& OutDrawElements, int32 LayerId, const FWidgetStyle& InWidgetStyle, bool bParentEnabled) const override
+	{
+		LayerId = SPanel::OnPaint(Args, AllottedGeometry, MyClippingRect, OutDrawElements, LayerId, InWidgetStyle, bParentEnabled);
+		LayerId = ScrollyZoomy.PaintSoftwareCursorIfNeeded(AllottedGeometry, MyClippingRect, OutDrawElements, LayerId);
+
+		if (PickedWidgets.Num() > 0)
+		{
+			TSharedPtr<FWidgetReflectorNodeBase> Window = SnapshotDataPtr->GetWindow(SelectedWindowIndex);
+			if (Window.IsValid())
+			{
+				const FVector2D RootDrawOffset = PhysicalOffset - Window->GetAccumulatedLayoutTransform().GetTranslation();
+				if (bIsPicking)
+				{
+					const FLinearColor TopmostWidgetColor(1.0f, 0.0f, 0.0f);
+					const FLinearColor LeafmostWidgetColor(0.0f, 1.0f, 0.0f);
+
+					for (int32 WidgetIndex = 0; WidgetIndex < PickedWidgets.Num(); ++WidgetIndex)
+					{
+						const TSharedRef<FWidgetReflectorNodeBase>& PickedWidget = PickedWidgets[WidgetIndex];
+						const float ColorFactor = static_cast<float>(WidgetIndex)/PickedWidgets.Num();
+						const FLinearColor Tint(1.0f - ColorFactor, ColorFactor, 0.0f, 1.0f);
+
+						FSlateDrawElement::MakeBox(
+							OutDrawElements,
+							++LayerId,
+							AllottedGeometry.ToPaintGeometry(RootDrawOffset + PickedWidget->GetAccumulatedLayoutTransform().GetTranslation(), TransformPoint(PickedWidget->GetAccumulatedLayoutTransform().GetScale(), PickedWidget->GetLocalSize())),
+							FCoreStyle::Get().GetBrush(TEXT("Debug.Border")),
+							MyClippingRect,
+							ESlateDrawEffect::None,
+							FMath::Lerp(TopmostWidgetColor, LeafmostWidgetColor, ColorFactor)
+						);
+					}
+				}
+				else
+				{
+					const TSharedRef<FWidgetReflectorNodeBase>& PickedWidget = PickedWidgets.Last();
+
+					FSlateDrawElement::MakeBox(
+						OutDrawElements,
+						++LayerId,
+						AllottedGeometry.ToPaintGeometry(RootDrawOffset + PickedWidget->GetAccumulatedLayoutTransform().GetTranslation(), TransformPoint(PickedWidget->GetAccumulatedLayoutTransform().GetScale(), PickedWidget->GetLocalSize())),
+						FCoreStyle::Get().GetBrush(TEXT("Debug.Border")),
+						MyClippingRect,
+						ESlateDrawEffect::None,
+						PickedWidget->GetTint()
+					);
+				}
+			}
+		}
+
+		return LayerId;
+	}
+
+	virtual bool ScrollBy(const FVector2D& Offset) override
+	{
+		const FVector2D PrevPhysicalOffset = PhysicalOffset;
+		PhysicalOffset += Offset;
+
+		const TSharedRef<SWidget>& ChildWidget = ChildSlot.GetWidget();
+		const FVector2D& WidgetDesiredSize = ChildWidget->GetDesiredSize();
+		ClampViewOffset(WidgetDesiredSize, CachedSize);
+
+		return PhysicalOffset != PrevPhysicalOffset;
+	}
+
+	virtual bool ZoomBy(const float Amount) override
+	{
+		return false;
+	}
+
+	float GetZoomLevel() const
+	{
+		return 1.0f;
+	}
+
+private:
+	void ClampViewOffset(const FVector2D& ViewportSize, const FVector2D& LocalSize)
+	{
+		PhysicalOffset.X = ClampViewOffsetAxis(ViewportSize.X, LocalSize.X, PhysicalOffset.X);
+		PhysicalOffset.Y = ClampViewOffsetAxis(ViewportSize.Y, LocalSize.Y, PhysicalOffset.Y);
+	}
+
+	float ClampViewOffsetAxis(const float ViewportSize, const float LocalSize, const float CurrentOffset)
+	{
+		if (ViewportSize <= LocalSize)
+		{
+			// If the viewport is smaller than the available size, then we can't be scrolled
+			return 0.0f;
+		}
+
+		// Given the size of the viewport, and the current size of the window, work how far we can scroll
+		// Note: This number is negative since scrolling down/right moves the viewport up/left
+		const float MaxScrollOffset = LocalSize - ViewportSize;
+
+		// Clamp the left/top edge
+		if (CurrentOffset < MaxScrollOffset)
+		{
+			return MaxScrollOffset;
+		}
+
+		// Clamp the right/bottom edge
+		if (CurrentOffset > 0.0f)
+		{
+			return 0.0f;
+		}
+
+		return CurrentOffset;
+	}
+
+	FVector2D PhysicalOffset;
+	mutable FVector2D CachedSize;
+
+	FScrollableSnapshotImageSlot ChildSlot;
+	FScrollyZoomy ScrollyZoomy;
+
+	/** Snapshot data we're visualizing */
+	const FWidgetSnapshotData* SnapshotDataPtr;
+
+	/** Index of the window we're currently viewing */
+	int32 SelectedWindowIndex;
+
+	SWidgetSnapshotVisualizer::FOnWidgetPathPicked OnWidgetPathPicked;
+
+	bool bIsPicking;
+	TArray<TSharedRef<FWidgetReflectorNodeBase>> PickedWidgets;
+};
+
+
+FWidgetSnapshotData::~FWidgetSnapshotData()
+{
+	DestroyBrushes();
+}
+
+void FWidgetSnapshotData::ClearSnapshot()
+{
+	Reset();
+}
+
+void FWidgetSnapshotData::TakeSnapshot()
+{
+	TArray<TSharedRef<SWindow>> VisibleWindows;
+	FSlateApplication::Get().GetAllVisibleWindowsOrdered(VisibleWindows);
+
+	CreateSnapshot(VisibleWindows);
+}
+
+void FWidgetSnapshotData::CreateSnapshot(const TArray<TSharedRef<SWindow>>& VisibleWindows)
+{
+	Reset();
+	Reserve(VisibleWindows.Num());
+
+	for (const auto& VisibleWindow : VisibleWindows)
+	{
+		// Snapshot the current state of this window widget hierarchy
+		Windows.Add(FWidgetReflectorNodeUtils::NewSnapshotNodeTreeFrom(FArrangedWidget(VisibleWindow, VisibleWindow->GetWindowGeometryInScreen())));
+
+		// Screenshot the current window so we can pick against its current state
+		FWidgetSnapshotTextureData& TextureData = WindowTextureData[WindowTextureData.AddDefaulted()];
+		FSlateApplication::Get().TakeScreenshot(VisibleWindow, TextureData.ColorData, TextureData.Dimensions);
+	}
+
+	CreateBrushes();
+}
+
+bool FWidgetSnapshotData::SaveSnapshotToFile(const FString& InFilename) const
+{
+	TSharedRef<FJsonObject> RootJsonObject = SaveSnapshotAsJson();
+
+	FArchive* const FileAr = IFileManager::Get().CreateFileWriter(*InFilename);
+	if (FileAr)
+	{
+		typedef TJsonWriter<TCHAR, TCondensedJsonPrintPolicy<TCHAR>> FStringWriter;
+		typedef TJsonWriterFactory<TCHAR, TCondensedJsonPrintPolicy<TCHAR>> FStringWriterFactory;
+
+		TSharedRef<FStringWriter> Writer = FStringWriterFactory::Create(FileAr);
+		FJsonSerializer::Serialize(RootJsonObject, Writer);
+		FileAr->Close();
+		return true;
+	}
+
+	return false;
+}
+
+void FWidgetSnapshotData::SaveSnapshotToBuffer(TArray<uint8>& OutData) const
+{
+	TSharedRef<FJsonObject> RootJsonObject = SaveSnapshotAsJson();
+
+	typedef TJsonWriter<TCHAR, TCondensedJsonPrintPolicy<TCHAR>> FStringWriter;
+	typedef TJsonWriterFactory<TCHAR, TCondensedJsonPrintPolicy<TCHAR>> FStringWriterFactory;
+
+	FArrayWriter TmpJsonData;
+	TSharedRef<FStringWriter> Writer = FStringWriterFactory::Create(&TmpJsonData);
+	FJsonSerializer::Serialize(RootJsonObject, Writer);
+
+	OutData.Reset();
+	FMemoryWriter BufferWriter(OutData);
+
+	int32 UncompressedDataSize = TmpJsonData.Num();
+	BufferWriter << UncompressedDataSize;
+
+	BufferWriter.SerializeCompressed(TmpJsonData.GetData(), TmpJsonData.Num(), COMPRESS_ZLIB);
+}
+
+TSharedRef<FJsonObject> FWidgetSnapshotData::SaveSnapshotAsJson() const
+{
+	check(Windows.Num() == WindowTextureData.Num());
+
+	TSharedRef<FJsonObject> RootJsonObject = MakeShareable(new FJsonObject());
+
+	{
+		TArray<TSharedPtr<FJsonValue>> WindowsJsonArray;
+		for (const auto& Window : Windows)
+		{
+			check(Window->GetNodeType() == EWidgetReflectorNodeType::Snapshot);
+			WindowsJsonArray.Add(FSnapshotWidgetReflectorNode::ToJson(StaticCastSharedRef<FSnapshotWidgetReflectorNode>(Window.ToSharedRef())));
+		}
+		RootJsonObject->SetArrayField(TEXT("Windows"), WindowsJsonArray);
+	}
+
+	{
+		TArray<TSharedPtr<FJsonValue>> TexturesJsonArray;
+		for (const FWidgetSnapshotTextureData& TextureData : WindowTextureData)
+		{
+			TSharedRef<FJsonObject> TextureDataJsonObject = MakeShareable(new FJsonObject());
+
+			{
+				TArray<TSharedPtr<FJsonValue>> StructJsonArray;
+				StructJsonArray.Add(MakeShareable(new FJsonValueNumber(TextureData.Dimensions.X)));
+				StructJsonArray.Add(MakeShareable(new FJsonValueNumber(TextureData.Dimensions.Y)));
+				TextureDataJsonObject->SetArrayField(TEXT("Dimensions"), StructJsonArray);
+			}
+
+			{
+				// This is raw texture data - compress it before we encode it to save space
+				const int32 UncompressedDataSizeBytes = TextureData.ColorData.Num() * sizeof(FColor);
+				TArray<uint8> CompressedDataBuffer;
+				CompressedDataBuffer.AddZeroed(FCompression::CompressMemoryBound(COMPRESS_ZLIB, UncompressedDataSizeBytes));
+				int32 CompressedDataSize = CompressedDataBuffer.Num();
+				if (FCompression::CompressMemory(COMPRESS_ZLIB, CompressedDataBuffer.GetData(), CompressedDataSize, TextureData.ColorData.GetData(), UncompressedDataSizeBytes))
+				{
+					TextureDataJsonObject->SetBoolField(TEXT("IsCompressed"), true);
+					TextureDataJsonObject->SetNumberField(TEXT("UncompressedSize"), UncompressedDataSizeBytes);
+
+					// FCompression::CompressMemory updates CompressedDataSize with the actual size - we may have to shrink our buffer count now
+					CompressedDataBuffer.SetNum(CompressedDataSize, false);
+
+					const FString EncodedTextureData = FBase64::Encode(CompressedDataBuffer);
+					TextureDataJsonObject->SetStringField(TEXT("TextureData"), EncodedTextureData);
+				}
+				else
+				{
+					TextureDataJsonObject->SetBoolField(TEXT("IsCompressed"), false);
+
+					// Failed to compress... use the raw texture data
+					TArray<uint8> TextureDataBytes;
+					TextureDataBytes.Append(reinterpret_cast<const uint8*>(TextureData.ColorData.GetData()), TextureData.ColorData.Num() * sizeof(FColor));
+
+					const FString EncodedTextureData = FBase64::Encode(TextureDataBytes);
+					TextureDataJsonObject->SetStringField(TEXT("TextureData"), EncodedTextureData);
+				}
+			}
+
+			TexturesJsonArray.Add(MakeShareable(new FJsonValueObject(TextureDataJsonObject)));
+		}
+		RootJsonObject->SetArrayField(TEXT("Textures"), TexturesJsonArray);
+	}
+
+	return RootJsonObject;
+}
+
+bool FWidgetSnapshotData::LoadSnapshotFromFile(const FString& InFilename)
+{
+	bool bJsonLoaded = false;
+	TSharedPtr<FJsonObject> RootJsonObject;
+
+	{
+		FArchive* FileAr = IFileManager::Get().CreateFileReader(*InFilename);
+		if (FileAr)
+		{
+			typedef TJsonReader<TCHAR> FJsonReader;
+			typedef TJsonReaderFactory<TCHAR> FJsonReaderFactory;
+
+			TSharedRef<FJsonReader> Reader = FJsonReaderFactory::Create(FileAr);
+			bJsonLoaded = FJsonSerializer::Deserialize(Reader, RootJsonObject);
+			FileAr->Close();
+			FileAr = nullptr;
+		}
+	}
+
+	if (bJsonLoaded)
+	{
+		check(RootJsonObject.IsValid());
+		LoadSnapshotFromJson(RootJsonObject.ToSharedRef());
+		return true;
+	}
+
+	return false;
+}
+
+void FWidgetSnapshotData::LoadSnapshotFromBuffer(const TArray<uint8>& InData)
+{
+	int32 UncompressedDataSize = 0;
+	TArray<uint8> UncompressedData;
+
+	{
+		FMemoryReader BufferReader(InData);
+
+		BufferReader << UncompressedDataSize;
+
+		UncompressedData.AddZeroed(UncompressedDataSize);
+		BufferReader.SerializeCompressed(UncompressedData.GetData(), 0, COMPRESS_ZLIB);
+	}
+
+	bool bJsonLoaded = false;
+	TSharedPtr<FJsonObject> RootJsonObject;
+
+	if (UncompressedData.Num() > 0)
+	{
+		typedef TJsonReader<TCHAR> FJsonReader;
+		typedef TJsonReaderFactory<TCHAR> FJsonReaderFactory;
+
+		FMemoryReader UncompressedDataReader(UncompressedData);
+		TSharedRef<FJsonReader> Reader = FJsonReaderFactory::Create(&UncompressedDataReader);
+		bJsonLoaded = FJsonSerializer::Deserialize(Reader, RootJsonObject);
+	}
+
+	if (bJsonLoaded)
+	{
+		check(RootJsonObject.IsValid());
+		LoadSnapshotFromJson(RootJsonObject.ToSharedRef());
+	}
+}
+
+void FWidgetSnapshotData::LoadSnapshotFromJson(const TSharedRef<FJsonObject>& InRootJsonObject)
+{
+	Reset();
+
+	{
+		const TArray<TSharedPtr<FJsonValue>>& WindowsJsonArray = InRootJsonObject->GetArrayField(TEXT("Windows"));
+		for (const TSharedPtr<FJsonValue>& WindowJsonValue : WindowsJsonArray)
+		{
+			Windows.Add(FSnapshotWidgetReflectorNode::FromJson(WindowJsonValue.ToSharedRef()));
+		}
+	}
+
+	{
+		const TArray<TSharedPtr<FJsonValue>>& TexturesJsonArray = InRootJsonObject->GetArrayField(TEXT("Textures"));
+		for (const TSharedPtr<FJsonValue>& TextureDataJsonValue : TexturesJsonArray)
+		{
+			const TSharedPtr<FJsonObject>& TextureDataJsonObject = TextureDataJsonValue->AsObject();
+			check(TextureDataJsonObject.IsValid());
+
+			FWidgetSnapshotTextureData TextureData;
+
+			{
+				const TArray<TSharedPtr<FJsonValue>>& StructJsonArray = TextureDataJsonObject->GetArrayField(TEXT("Dimensions"));
+				check(StructJsonArray.Num() == 2);
+
+				TextureData.Dimensions.X = FMath::TruncToInt(StructJsonArray[0]->AsNumber());
+				TextureData.Dimensions.Y = FMath::TruncToInt(StructJsonArray[1]->AsNumber());
+			}
+
+			{
+				const FString EncodedTextureData = TextureDataJsonObject->GetStringField(TEXT("TextureData"));
+				TArray<uint8> DecodedTextureDataBytes;
+				FBase64::Decode(EncodedTextureData, DecodedTextureDataBytes);
+
+				const bool bIsCompressed = TextureDataJsonObject->GetBoolField(TEXT("IsCompressed"));
+				if (bIsCompressed)
+				{
+					const int32 UncompressedDataSizeBytes = FMath::TruncToInt(TextureDataJsonObject->GetNumberField(TEXT("UncompressedSize")));
+					TextureData.ColorData.AddZeroed(UncompressedDataSizeBytes / sizeof(FColor));
+
+					FCompression::UncompressMemory(COMPRESS_ZLIB, TextureData.ColorData.GetData(), UncompressedDataSizeBytes, DecodedTextureDataBytes.GetData(), DecodedTextureDataBytes.Num());
+				}
+				else
+				{
+					TextureData.ColorData.Append(reinterpret_cast<const FColor*>(DecodedTextureDataBytes.GetData()), DecodedTextureDataBytes.Num() / sizeof(FColor));
+				}
+			}
+
+			WindowTextureData.Add(TextureData);
+		}
+	}
+
+	CreateBrushes();
+}
+
+bool FWidgetSnapshotData::IsEmpty() const
+{
+	return Windows.Num() == 0;
+}
+
+int32 FWidgetSnapshotData::Num() const
+{
+	return Windows.Num();
+}
+
+const TArray<TSharedPtr<FWidgetReflectorNodeBase>>& FWidgetSnapshotData::GetWindowsPtr() const
+{
+	return Windows;
+}
+
+TArray<TSharedRef<FWidgetReflectorNodeBase>> FWidgetSnapshotData::GetWindowsRef() const
+{
+	TArray<TSharedRef<FWidgetReflectorNodeBase>> RetWindows;
+	RetWindows.Reserve(Windows.Num());
+	for (const auto& Window : Windows)
+	{
+		RetWindows.Add(Window.ToSharedRef());
+	}
+	return RetWindows;
+}
+
+TSharedPtr<FWidgetReflectorNodeBase> FWidgetSnapshotData::GetWindow(const int32 WindowIndex) const
+{
+	return (Windows.IsValidIndex(WindowIndex)) 
+		? TSharedPtr<FWidgetReflectorNodeBase>(Windows[WindowIndex]) 
+		: TSharedPtr<FWidgetReflectorNodeBase>(nullptr);
+}
+
+const FSlateBrush* FWidgetSnapshotData::GetBrush(const int32 WindowIndex) const
+{
+	return (WindowTextureBrushes.IsValidIndex(WindowIndex)) ? WindowTextureBrushes[WindowIndex].Get() : nullptr;
+}
+
+void FWidgetSnapshotData::CreateBrushes()
+{
+	DestroyBrushes();
+
+	WindowTextureBrushes.Reserve(WindowTextureData.Num());
+
+	static int32 TextureIndex = 0;
+	for (const auto& TextureData : WindowTextureData)
+	{
+		if (TextureData.ColorData.Num() > 0)
+		{
+			TArray<uint8> TextureDataAsBGRABytes;
+			TextureDataAsBGRABytes.Reserve(TextureData.ColorData.Num() * 4);
+			for (const FColor& PixelColor : TextureData.ColorData)
+			{
+				TextureDataAsBGRABytes.Add(PixelColor.B);
+				TextureDataAsBGRABytes.Add(PixelColor.G);
+				TextureDataAsBGRABytes.Add(PixelColor.R);
+				TextureDataAsBGRABytes.Add(PixelColor.A);
+			}
+
+			WindowTextureBrushes.Add(FSlateDynamicImageBrush::CreateWithImageData(
+				*FString::Printf(TEXT("FWidgetSnapshotData_WindowTextureBrush_%d"), TextureIndex++), 
+				FVector2D(TextureData.Dimensions.X, TextureData.Dimensions.Y), 
+				TextureDataAsBGRABytes
+				));
+		}
+		else
+		{
+			WindowTextureBrushes.Add(nullptr);
+		}
+	}
+}
+
+void FWidgetSnapshotData::DestroyBrushes()
+{
+	for (const auto& WindowTextureBrush : WindowTextureBrushes)
+	{
+		if (WindowTextureBrush.IsValid())
+		{
+			WindowTextureBrush->ReleaseResource();
+		}
+	}
+
+	WindowTextureBrushes.Reset();
+}
+
+void FWidgetSnapshotData::Reserve(const int32 NumWindows)
+{
+	Windows.Reserve(NumWindows);
+	WindowTextureData.Reserve(NumWindows);
+	WindowTextureBrushes.Reserve(NumWindows);
+}
+
+void FWidgetSnapshotData::Reset()
+{
+	DestroyBrushes();
+
+	Windows.Reset();
+	WindowTextureData.Reset();
+	WindowTextureBrushes.Reset();
+}
+
+
+void SWidgetSnapshotVisualizer::Construct(const FArguments& InArgs)
+{
+	SnapshotDataPtr = InArgs._SnapshotData;
+	check(SnapshotDataPtr);
+
+	ChildSlot
+	[
+		SNew(SVerticalBox)
+
+		+SVerticalBox::Slot()
+		.AutoHeight()
+		[
+			SNew(SHorizontalBox)
+
+			+SHorizontalBox::Slot()
+			.AutoWidth()
+			.Padding(FMargin(5.0f, 0.0f))
+			[
+				SAssignNew(WindowPickerCombo, SComboBox<TSharedPtr<FWidgetReflectorNodeBase>>)
+				.OptionsSource(&SnapshotDataPtr->GetWindowsPtr())
+				.OnSelectionChanged(this, &SWidgetSnapshotVisualizer::OnWindowSelectionChanged)
+				.OnGenerateWidget(this, &SWidgetSnapshotVisualizer::GenerateWindowPickerComboItem)
+				[
+					SNew(STextBlock)
+					.Text(this, &SWidgetSnapshotVisualizer::GetSelectedWindowComboItemText)
+				]
+			]
+
+			+SHorizontalBox::Slot()
+			.AutoWidth()
+			.Padding(FMargin(5.0f, 0.0f))
+			[
+				SNew(SButton)
+				.Text(this, &SWidgetSnapshotVisualizer::GetPickWidgetText)
+				.ButtonColorAndOpacity(this, &SWidgetSnapshotVisualizer::GetPickWidgetColor)
+				.OnClicked(this, &SWidgetSnapshotVisualizer::OnPickWidgetClicked)
+			]
+
+#if SLATE_REFLECTOR_HAS_DESKTOP_PLATFORM
+			+SHorizontalBox::Slot()
+			.AutoWidth()
+			.Padding(FMargin(5.0f, 0.0f))
+			[
+				SNew(SButton)
+				.Text(LOCTEXT("SaveSnapshotButtonText", "Save Snapshot"))
+				.OnClicked(this, &SWidgetSnapshotVisualizer::OnSaveSnapshotClicked)
+			]
+#endif // SLATE_REFLECTOR_HAS_DESKTOP_PLATFORM
+		]
+
+		+SVerticalBox::Slot()
+		[
+			SAssignNew(SnapshotImage, SScrollableSnapshotImage)
+			.SnapshotData(InArgs._SnapshotData)
+			.OnWidgetPathPicked(InArgs._OnWidgetPathPicked)
+		]
+	];
+
+	SnapshotDataUpdated();
+}
+
+void SWidgetSnapshotVisualizer::SnapshotDataUpdated()
+{
+	if (SnapshotImage.IsValid())
+	{
+		SnapshotImage->SetSelectedWindowIndex(0);
+	}
+
+	if (WindowPickerCombo.IsValid())
+	{
+		WindowPickerCombo->RefreshOptions();
+		WindowPickerCombo->SetSelectedItem(SnapshotDataPtr->GetWindow(0));
+	}
+}
+
+FReply SWidgetSnapshotVisualizer::OnPreviewKeyDown(const FGeometry& MyGeometry, const FKeyEvent& InKeyEvent)
+{
+	if (SnapshotImage.IsValid() && InKeyEvent.GetKey() == EKeys::Escape)
+	{
+		SnapshotImage->SetIsPicking(false);
+	}
+
+	return FReply::Unhandled();
+}
+
+void SWidgetSnapshotVisualizer::OnWindowSelectionChanged(TSharedPtr<FWidgetReflectorNodeBase> InWindow, ESelectInfo::Type InReason)
+{
+	int32 SelectedWindowIndex = INDEX_NONE;
+	SnapshotDataPtr->GetWindowsPtr().Find(InWindow, SelectedWindowIndex);
+
+	if (SnapshotImage.IsValid())
+	{
+		SnapshotImage->SetSelectedWindowIndex(SelectedWindowIndex);
+	}
+}
+
+FText SWidgetSnapshotVisualizer::GetWindowPickerComboItemText(TSharedPtr<FWidgetReflectorNodeBase> InWindow)
+{
+	return FText::Format(LOCTEXT("WidgetComboItemFmt", "{0} - {1}"), InWindow->GetWidgetType(), InWindow->GetWidgetReadableLocation());
+}
+
+FText SWidgetSnapshotVisualizer::GetSelectedWindowComboItemText() const
+{
+	const int32 SelectedWindowIndex = SnapshotImage.IsValid() ? SnapshotImage->GetSelectedWindowIndex() : INDEX_NONE;
+	TSharedPtr<FWidgetReflectorNodeBase> SelectedWindowPtr = SnapshotDataPtr->GetWindow(SelectedWindowIndex);
+	return (SelectedWindowPtr.IsValid()) ? GetWindowPickerComboItemText(SelectedWindowPtr) : FText::GetEmpty();
+}
+
+TSharedRef<SWidget> SWidgetSnapshotVisualizer::GenerateWindowPickerComboItem(TSharedPtr<FWidgetReflectorNodeBase> InWindow) const
+{
+	return SNew(STextBlock)
+		.Text(GetWindowPickerComboItemText(InWindow));
+}
+
+FText SWidgetSnapshotVisualizer::GetPickWidgetText() const
+{
+	const bool bIsPicking = SnapshotImage.IsValid() && SnapshotImage->GetIsPicking();
+	return (bIsPicking) ? LOCTEXT("PickingWidget", "Picking (Esc to Stop)") : LOCTEXT("PickSnapshotWidget", "Pick Snapshot Widget");
+}
+
+FSlateColor SWidgetSnapshotVisualizer::GetPickWidgetColor() const
+{
+	static const FName SelectionColor("SelectionColor");
+
+	const bool bIsPicking = SnapshotImage.IsValid() && SnapshotImage->GetIsPicking();
+	return bIsPicking
+		? FCoreStyle::Get().GetSlateColor(SelectionColor)
+		: FLinearColor::White;
+}
+
+FReply SWidgetSnapshotVisualizer::OnPickWidgetClicked()
+{
+	if (SnapshotImage.IsValid())
+	{
+		SnapshotImage->SetIsPicking(!SnapshotImage->GetIsPicking());
+	}
+	return FReply::Handled();
+}
+
+#if SLATE_REFLECTOR_HAS_DESKTOP_PLATFORM
+
+FReply SWidgetSnapshotVisualizer::OnSaveSnapshotClicked()
+{
+	IDesktopPlatform* DesktopPlatform = FDesktopPlatformModule::Get();
+
+	if (DesktopPlatform)
+	{
+		TSharedPtr<SWindow> ParentWindow = FSlateApplication::Get().FindWidgetWindow(SharedThis(this));
+
+		TArray<FString> SaveFilenames;
+		const bool bOpened = DesktopPlatform->SaveFileDialog(
+			(ParentWindow.IsValid()) ? ParentWindow->GetNativeWindow()->GetOSWindowHandle() : nullptr,
+			LOCTEXT("SaveSnapshotDialogTitle", "Save Widget Snapshot").ToString(),
+			FPaths::GameAgnosticSavedDir(),
+			TEXT(""),
+			TEXT("Slate Widget Snapshot (*.widgetsnapshot)|*.widgetsnapshot"),
+			EFileDialogFlags::None,
+			SaveFilenames
+			);
+
+		if (SaveFilenames.Num() > 0)
+		{
+			SnapshotDataPtr->SaveSnapshotToFile(SaveFilenames[0]);
+		}
+	}
+
+	return FReply::Handled();
+}
+
+#endif // SLATE_REFLECTOR_HAS_DESKTOP_PLATFORM
+
+#undef LOCTEXT_NAMESPACE
diff --git a/Engine/Source/Developer/SlateReflector/Private/Widgets/SWidgetSnapshotVisualizer.h b/Engine/Source/Developer/SlateReflector/Private/Widgets/SWidgetSnapshotVisualizer.h
new file mode 100644
--- /dev/null
+++ b/Engine/Source/Developer/SlateReflector/Private/Widgets/SWidgetSnapshotVisualizer.h
@@ -0,0 +1,159 @@
+// Copyright 1998-2015 Epic Games, Inc. All Rights Reserved.
+
+#pragma once
+
+class FWidgetReflectorNodeBase;
+class SScrollableSnapshotImage;
+
+/** 
+ * The raw texture data from taking a screenshot of a Slate widget (typically the root window)
+ */
+struct FWidgetSnapshotTextureData
+{
+	/** The dimensions of the texture */
+	FIntVector Dimensions;
+
+	/** The raw color data for the texture (BGRA) */
+	TArray<FColor> ColorData;
+};
+
+/** 
+ * All of the data relating to a single widget hierarchy snapshot
+ */
+class FWidgetSnapshotData
+{
+public:
+	/** Destructor */
+	~FWidgetSnapshotData();
+
+	/** Clear the current snapshot data so that we can reclaim the memory */
+	void ClearSnapshot();
+
+	/** Take a snapshot of all of the windows that are currently open */
+	void TakeSnapshot();
+
+	/** Create a snapshot of the given windows */
+	void CreateSnapshot(const TArray<TSharedRef<SWindow>>& VisibleWindows);
+
+	/** Save this snapshot data to the given file. The data will be saved as uncompressed JSON data. */
+	bool SaveSnapshotToFile(const FString& InFilename) const;
+
+	/** Save this snapshot data to the given buffer. The data will be saved as zlib compressed JSON data. */
+	void SaveSnapshotToBuffer(TArray<uint8>& OutData) const;
+
+	/** Create a JSON object that represents the snapshot data. */
+	TSharedRef<FJsonObject> SaveSnapshotAsJson() const;
+
+	/** Populate this snapshot data from the given file. */
+	bool LoadSnapshotFromFile(const FString& InFilename);
+
+	/** Populate this snapshot data from the given buffer. */
+	void LoadSnapshotFromBuffer(const TArray<uint8>& InData);
+
+	/** Populate this snapshot data from the given JSON object */
+	void LoadSnapshotFromJson(const TSharedRef<FJsonObject>& InRootJsonObject);
+
+	/** Check to see whether this snapshot is empty (contains no windows) */
+	bool IsEmpty() const;
+
+	/** Get the number of windows this snapshot contains */
+	int32 Num() const;
+
+	/** Get the internal windows array */
+	const TArray<TSharedPtr<FWidgetReflectorNodeBase>>& GetWindowsPtr() const;
+
+	/** Get the internal windows array, but with each window converted to a TSharedRef */
+	TArray<TSharedRef<FWidgetReflectorNodeBase>> GetWindowsRef() const;
+
+	/** Get the window for the given index, or null if the index is invalid */
+	TSharedPtr<FWidgetReflectorNodeBase> GetWindow(const int32 WindowIndex) const;
+
+	/** Get the brush for the given index, or null if the index is invalid */
+	const FSlateBrush* GetBrush(const int32 WindowIndex) const;
+
+private:
+	/** Create the dynamic Slate brushes from the texture data for each window */
+	void CreateBrushes();
+
+	/** Destroy the dynamic Slate brushes for each window */
+	void DestroyBrushes();
+
+	/** Reserve space in all of our internal arrays for the given number of entries */
+	void Reserve(const int32 NumWindows);
+
+	/** Reset our internal arrays */
+	void Reset();
+
+	/** Array of root level windows, each containing a tree of widget nodes */
+	TArray<TSharedPtr<FWidgetReflectorNodeBase>> Windows;
+
+	/** Contains a texture data entry for each entry in Windows */
+	TArray<FWidgetSnapshotTextureData> WindowTextureData;
+
+	/** Contains a dynamic brush pointer for each entry in WindowTextureData */
+	TArray<TSharedPtr<FSlateDynamicImageBrush>> WindowTextureBrushes;
+};
+
+/**
+ * Visualizer to handle viewing and picking from a widget hierarchy snapshot
+ */
+class SWidgetSnapshotVisualizer : public SCompoundWidget
+{
+public:
+	DECLARE_DELEGATE_OneParam(FOnWidgetPathPicked, const TArray<TSharedRef<FWidgetReflectorNodeBase>>& /*PickedWidgetPath*/);
+
+	SLATE_BEGIN_ARGS(SWidgetSnapshotVisualizer)
+		: _SnapshotData(nullptr)
+		{}
+
+		SLATE_ARGUMENT(const FWidgetSnapshotData*, SnapshotData)
+
+		SLATE_EVENT(FOnWidgetPathPicked, OnWidgetPathPicked);
+
+	SLATE_END_ARGS()
+
+	void Construct(const FArguments& InArgs);
+
+	/** Called when the snapshot data we're observing is changed. This allows us to update our view. */
+	void SnapshotDataUpdated();
+
+	// SWidget interface
+	virtual FReply OnPreviewKeyDown(const FGeometry& MyGeometry, const FKeyEvent& InKeyEvent) override;
+
+private:
+	/** Called when the selected window is changed in the combo box. Used to update our view. */
+	void OnWindowSelectionChanged(TSharedPtr<FWidgetReflectorNodeBase> InWindow, ESelectInfo::Type InReason);
+
+	/** Get the combo item text to use for the given window */
+	static FText GetWindowPickerComboItemText(TSharedPtr<FWidgetReflectorNodeBase> InWindow);
+
+	/** Get the combo item text for the currently selected window */
+	FText GetSelectedWindowComboItemText() const;
+
+	/** Create a widget for the items in the window picker combo box */
+	TSharedRef<SWidget> GenerateWindowPickerComboItem(TSharedPtr<FWidgetReflectorNodeBase> InWindow) const;
+
+	/** Get the current text to use for the "Pick Snapshot Widget" button */
+	FText GetPickWidgetText() const;
+
+	/** Get the current color to use for the "Pick Snapshot Widget" button */
+	FSlateColor GetPickWidgetColor() const;
+
+	/** Called when the "Pick Snapshot Widget" button is clicked */
+	FReply OnPickWidgetClicked();
+
+#if SLATE_REFLECTOR_HAS_DESKTOP_PLATFORM
+	/** Called when the "Save Snapshot" button is clicked */
+	FReply OnSaveSnapshotClicked();
+#endif // SLATE_REFLECTOR_HAS_DESKTOP_PLATFORM
+
+private:
+	/** Snapshot data we're visualizing */
+	const FWidgetSnapshotData* SnapshotDataPtr;
+
+	/** Window picker combo box */
+	TSharedPtr<SComboBox<TSharedPtr<FWidgetReflectorNodeBase>>> WindowPickerCombo;
+
+	/** Snapshot image */
+	TSharedPtr<SScrollableSnapshotImage> SnapshotImage;
+};
diff --git a/Engine/Source/Developer/SlateReflector/Public/ISlateReflectorModule.h b/Engine/Source/Developer/SlateReflector/Public/ISlateReflectorModule.h
--- a/Engine/Source/Developer/SlateReflector/Public/ISlateReflectorModule.h
+++ b/Engine/Source/Developer/SlateReflector/Public/ISlateReflectorModule.h
@@ -19,34 +19,6 @@ class ISlateReflectorModule
 {
 public:
 
-	/**
-	 * Creates a widget reflector widget.
-	 *
-	 * @return The widget.
-	 */
-	virtual TSharedRef<SWidget> GetWidgetReflector() = 0;
-
-	/**
-	 * Create an atlas visualizer widget using the given atlas provider.
-	 *
-	 * @return The widget.
-	 */
-	virtual TSharedRef<SWidget> GetAtlasVisualizer( ISlateAtlasProvider* InAtlasProvider ) = 0;
-
-	/**
-	 * Create an atlas visualizer widget for the texture atlas provider used by the current renderer.
-	 *
-	 * @return The widget.
-	 */
-	virtual TSharedRef<SWidget> GetTextureAtlasVisualizer() = 0;
-
-	/**
-	 * Create an atlas visualizer widget for the font atlas provider used by the current renderer.
-	 *
-	 * @return The widget.
-	 */
-	virtual TSharedRef<SWidget> GetFontAtlasVisualizer() = 0;
-
 	/**
 	 * Display the widget reflector, either spawned from a tab manager, or in a new window if the tab manager can't be used
 	 */
@@ -67,7 +39,7 @@ public:
 	 *
 	 * @param WorkspaceGroup The workspace group to insert the tab into.
 	 */
-	virtual void RegisterTabSpawner( const TSharedRef<FWorkspaceItem>& WorkspaceGroup ) = 0;
+	virtual void RegisterTabSpawner( const TSharedPtr<FWorkspaceItem>& WorkspaceGroup ) = 0;
 
 	/** Unregisters the tab spawner for the widget reflector. */
 	virtual void UnregisterTabSpawner() = 0;
diff --git a/Engine/Source/Developer/SlateReflector/SlateReflector.Build.cs b/Engine/Source/Developer/SlateReflector/SlateReflector.Build.cs
--- a/Engine/Source/Developer/SlateReflector/SlateReflector.Build.cs
+++ b/Engine/Source/Developer/SlateReflector/SlateReflector.Build.cs
@@ -18,6 +18,7 @@ public class SlateReflector : ModuleRules
 				"InputCore",
 				"Slate",
 				"SlateCore",
+				"Json",
 			}
 		);
 
@@ -28,5 +29,25 @@ public class SlateReflector : ModuleRules
 				"Developer/SlateReflector/Private/Widgets",
 			}
 		);
+
+		// DesktopPlatform is only available for Editor and Program targets (running on a desktop platform)
+		bool IsDesktopPlatformType = Target.Platform == UnrealBuildTool.UnrealTargetPlatform.Win32
+			|| Target.Platform == UnrealBuildTool.UnrealTargetPlatform.Win64
+			|| Target.Platform == UnrealBuildTool.UnrealTargetPlatform.Mac
+			|| Target.Platform == UnrealBuildTool.UnrealTargetPlatform.Linux;
+		if (Target.Type == TargetRules.TargetType.Editor || (Target.Type == TargetRules.TargetType.Program && IsDesktopPlatformType))
+		{
+			Definitions.Add("SLATE_REFLECTOR_HAS_DESKTOP_PLATFORM=1");
+
+			PrivateDependencyModuleNames.AddRange(
+				new string[] {
+					"DesktopPlatform",
+				}
+			);
+		}
+		else
+		{
+			Definitions.Add("SLATE_REFLECTOR_HAS_DESKTOP_PLATFORM=0");
+		}
 	}
 }
diff --git a/Engine/Source/Programs/CrashReporter/CrashReportClient/Private/CrashReportClientApp.cpp b/Engine/Source/Programs/CrashReporter/CrashReportClient/Private/CrashReportClientApp.cpp
--- a/Engine/Source/Programs/CrashReporter/CrashReportClient/Private/CrashReportClientApp.cpp
+++ b/Engine/Source/Programs/CrashReporter/CrashReportClient/Private/CrashReportClientApp.cpp
@@ -192,12 +192,7 @@ void RunCrashReportClient(const TCHAR* CommandLine)
 		// Debugging code
 		if (RunWidgetReflector)
 		{
-			FSlateApplication::Get().AddWindow(
-				SNew(SWindow)
-				.ClientSize(FVector2D(800, 600))
-				[
-					FModuleManager::LoadModuleChecked<ISlateReflectorModule>("SlateReflector").GetWidgetReflector()
-				]);
+			FModuleManager::LoadModuleChecked<ISlateReflectorModule>("SlateReflector").DisplayWidgetReflector();
 		}
 
 		// loop until the app is ready to quit
diff --git a/Engine/Source/Runtime/AppFramework/Private/Widgets/Testing/SPerfSuite.cpp b/Engine/Source/Runtime/AppFramework/Private/Widgets/Testing/SPerfSuite.cpp
--- a/Engine/Source/Runtime/AppFramework/Private/Widgets/Testing/SPerfSuite.cpp
+++ b/Engine/Source/Runtime/AppFramework/Private/Widgets/Testing/SPerfSuite.cpp
@@ -5,30 +5,46 @@
 
 #if !UE_BUILD_SHIPPING
 
+#include "SDockTab.h"
 #include "STableViewTesting.h"
 #include "ISlateReflectorModule.h"
 
 void SummonPerfTestSuite()
 {
-	FSlateApplication::Get().AddWindow
-	( 
-		SNew(SWindow)
-		.IsInitiallyMaximized(false)
-		.ScreenPosition( FVector2D(0,0) )
-		.ClientSize( FVector2D(1920,1200) )
-		[
-			MakeTableViewTesting()
-		]
-	);
+	// Need to load this module so we have the widget reflector tab available
+	FModuleManager::LoadModuleChecked<ISlateReflectorModule>("SlateReflector");
+
+	auto SpawnTableViewTesting = [](const FSpawnTabArgs&) -> TSharedRef<SDockTab>
+	{
+		return SNew(SDockTab)
+			.TabRole(ETabRole::NomadTab)
+			[
+				MakeTableViewTesting()
+			];
+	};
 
-	FSlateApplication::Get().AddWindow
+	FGlobalTabmanager::Get()->RegisterNomadTabSpawner("TableViewTesting", FOnSpawnTab::CreateLambda(SpawnTableViewTesting));
+	
+	TSharedRef<FTabManager::FLayout> Layout = FTabManager::NewLayout( "PerfTestSuite_Layout" )
+	->AddArea
 	(
-		SNew(SWindow)
-		.ClientSize(FVector2D(640,800))
-		[
-			FModuleManager::LoadModuleChecked<ISlateReflectorModule>("SlateReflector").GetWidgetReflector()
-		]
+		FTabManager::NewArea(1920,1200)
+		->Split
+		(
+			FTabManager::NewStack()
+			->AddTab("TableViewTesting", ETabState::OpenedTab)
+		)
+	)
+	->AddArea
+	(
+		FTabManager::NewArea(640,800)
+		->Split
+		(
+			FTabManager::NewStack()->AddTab("WidgetReflector", ETabState::OpenedTab)
+		)
 	);
+
+	FGlobalTabmanager::Get()->RestoreFrom(Layout, TSharedPtr<SWindow>());
 }
 
 #endif // #if !UE_BUILD_SHIPPING
diff --git a/Engine/Source/Runtime/AppFramework/Private/Widgets/Testing/STestSuite.cpp b/Engine/Source/Runtime/AppFramework/Private/Widgets/Testing/STestSuite.cpp
--- a/Engine/Source/Runtime/AppFramework/Private/Widgets/Testing/STestSuite.cpp
+++ b/Engine/Source/Runtime/AppFramework/Private/Widgets/Testing/STestSuite.cpp
@@ -5,6 +5,7 @@
 
 #if !UE_BUILD_SHIPPING
 
+#include "ISlateReflectorModule.h"
 #include "STableViewTesting.h"
 #include "SLayoutExample.h"
 #include "SWidgetGallery.h"
@@ -5704,6 +5705,9 @@ TSharedRef<SDockTab> SpawnWidgetGallery(const FSpawnTabArgs& Args)
 
 void RestoreSlateTestSuite()
 {
+	// Need to load this module so we have the widget reflector tab available
+	FModuleManager::LoadModuleChecked<ISlateReflectorModule>("SlateReflector");
+
 	FTestStyle::ResetToDefault();
 
 	FGlobalTabmanager::Get()->RegisterTabSpawner("TestSuite1", FOnSpawnTab::CreateStatic( &SpawnTestSuite1 ) );
