UMG/Slate - Exposing the Widget Caching flag to the widget reflector.  Also making it so invalidation is disabled when Picking Widgets is occuring, as invalidation causes that system of traversing the widget hierarchy to no longer function.  Since it pretends there are no children to avoid Slate doing more work than it needs to.

diff --git a/Engine/Source/Developer/SlateReflector/Private/Widgets/SWidgetReflector.cpp b/Engine/Source/Developer/SlateReflector/Private/Widgets/SWidgetReflector.cpp
--- a/Engine/Source/Developer/SlateReflector/Private/Widgets/SWidgetReflector.cpp
+++ b/Engine/Source/Developer/SlateReflector/Private/Widgets/SWidgetReflector.cpp
@@ -210,6 +210,11 @@ private:
 		{
 			bShowFocus = false;
 			SetUIMode(EWidgetReflectorUIMode::Live);
+			SInvalidationPanel::SetEnableWidgetCaching(false);
+		}
+		else
+		{
+			SInvalidationPanel::SetEnableWidgetCaching(true);
 		}
 
 		return FReply::Handled();
@@ -611,6 +616,31 @@ TSharedRef<SDockTab> SWidgetReflector::SpawnWidgetHierarchyTab(const FSpawnTabAr
 						SNew(STextBlock)
 						.Text(LOCTEXT("ShowFocus", "Show Focus"))
 					]
+
+					+ SHorizontalBox::Slot()
+					.AutoWidth()
+					.Padding(5.0f)
+					[
+						SNew(SCheckBox)
+						.Style(FCoreStyle::Get(), "ToggleButtonCheckbox")
+						.IsChecked_Lambda([&]()
+						{
+							return SInvalidationPanel::GetEnableWidgetCaching() ? ECheckBoxState::Checked : ECheckBoxState::Unchecked;
+						})
+						.OnCheckStateChanged_Lambda([&](const ECheckBoxState NewState)
+						{
+							SInvalidationPanel::SetEnableWidgetCaching(( NewState == ECheckBoxState::Checked ) ? true : false);
+						})
+						[
+							SNew(SBox)
+							.VAlign(VAlign_Center)
+							.HAlign(HAlign_Center)
+							[
+								SNew(STextBlock)
+								.Text(LOCTEXT("EnableWidgetCaching", "Widget Caching"))
+							]
+						]
+					]
 				]
 
 				+SHorizontalBox::Slot()
diff --git a/Engine/Source/Runtime/Slate/Private/Widgets/SInvalidationPanel.cpp b/Engine/Source/Runtime/Slate/Private/Widgets/SInvalidationPanel.cpp
--- a/Engine/Source/Runtime/Slate/Private/Widgets/SInvalidationPanel.cpp
+++ b/Engine/Source/Runtime/Slate/Private/Widgets/SInvalidationPanel.cpp
@@ -36,6 +36,16 @@ TAutoConsoleVariable<int32> EnableWidgetCaching(
 	true,
 	TEXT("Whether to attempt to cache any widgets through invalidation panels."));
 
+bool SInvalidationPanel::GetEnableWidgetCaching()
+{
+	return EnableWidgetCaching.GetValueOnGameThread() == 1;
+}
+
+void SInvalidationPanel::SetEnableWidgetCaching(bool bEnable)
+{
+	EnableWidgetCaching.AsVariable()->Set(bEnable);
+}
+
 #endif
 
 void SInvalidationPanel::Construct( const FArguments& InArgs )
diff --git a/Engine/Source/Runtime/Slate/Public/Widgets/SInvalidationPanel.h b/Engine/Source/Runtime/Slate/Public/Widgets/SInvalidationPanel.h
--- a/Engine/Source/Runtime/Slate/Public/Widgets/SInvalidationPanel.h
+++ b/Engine/Source/Runtime/Slate/Public/Widgets/SInvalidationPanel.h
@@ -45,6 +45,9 @@ public:
 #if !(UE_BUILD_SHIPPING || UE_BUILD_TEST)
 	static bool IsInvalidationDebuggingEnabled();
 	static void EnableInvalidationDebugging(bool bEnable);
+
+	static bool GetEnableWidgetCaching();
+	static void SetEnableWidgetCaching(bool bEnable);
 #endif
 
 private:
