build fix

diff --git a/Engine/Source/Developer/SlateReflector/Private/Widgets/SWidgetReflector.cpp b/Engine/Source/Developer/SlateReflector/Private/Widgets/SWidgetReflector.cpp
--- a/Engine/Source/Developer/SlateReflector/Private/Widgets/SWidgetReflector.cpp
+++ b/Engine/Source/Developer/SlateReflector/Private/Widgets/SWidgetReflector.cpp
@@ -616,33 +616,32 @@ TSharedRef<SDockTab> SWidgetReflector::SpawnWidgetHierarchyTab(const FSpawnTabAr
 						SNew(STextBlock)
 						.Text(LOCTEXT("ShowFocus", "Show Focus"))
 					]
-
-					+ SHorizontalBox::Slot()
-					.AutoWidth()
-					.Padding(5.0f)
+                ]
+				+ SHorizontalBox::Slot()
+				.AutoWidth()
+				.Padding(5.0f)
+				[
+					SNew(SCheckBox)
+					.Style(FCoreStyle::Get(), "ToggleButtonCheckbox")
+					.IsChecked_Lambda([&]()
+					{
+						return SInvalidationPanel::GetEnableWidgetCaching() ? ECheckBoxState::Checked : ECheckBoxState::Unchecked;
+					})
+					.OnCheckStateChanged_Lambda([&](const ECheckBoxState NewState)
+					{
+						SInvalidationPanel::SetEnableWidgetCaching(( NewState == ECheckBoxState::Checked ) ? true : false);
+					})
 					[
-						SNew(SCheckBox)
-						.Style(FCoreStyle::Get(), "ToggleButtonCheckbox")
-						.IsChecked_Lambda([&]()
-						{
-							return SInvalidationPanel::GetEnableWidgetCaching() ? ECheckBoxState::Checked : ECheckBoxState::Unchecked;
-						})
-						.OnCheckStateChanged_Lambda([&](const ECheckBoxState NewState)
-						{
-							SInvalidationPanel::SetEnableWidgetCaching(( NewState == ECheckBoxState::Checked ) ? true : false);
-						})
+						SNew(SBox)
+						.VAlign(VAlign_Center)
+						.HAlign(HAlign_Center)
 						[
-							SNew(SBox)
-							.VAlign(VAlign_Center)
-							.HAlign(HAlign_Center)
-							[
-								SNew(STextBlock)
-								.Text(LOCTEXT("EnableWidgetCaching", "Widget Caching"))
-							]
+							SNew(STextBlock)
+							.Text(LOCTEXT("EnableWidgetCaching", "Widget Caching"))
 						]
 					]
 				]
-
+				
 				+SHorizontalBox::Slot()
 				.AutoWidth()
 				.Padding(FMargin(5.0f, 0.0f))
