Speculative build fix for root editor build failure.

diff --git a/Engine/Source/Developer/SlateReflector/Private/Widgets/SWidgetReflector.cpp b/Engine/Source/Developer/SlateReflector/Private/Widgets/SWidgetReflector.cpp
--- a/Engine/Source/Developer/SlateReflector/Private/Widgets/SWidgetReflector.cpp
+++ b/Engine/Source/Developer/SlateReflector/Private/Widgets/SWidgetReflector.cpp
@@ -497,7 +497,7 @@ void SWidgetReflector::Construct( const FArguments& InArgs )
 				[
 					SNew(SCheckBox)
 					.Style( FCoreStyle::Get(), "ToggleButtonCheckbox" )
-					.IsEnabled(SLATE_STATS)
+					.IsEnabled(SLATE_STATS != 0)
 					.IsChecked_Static([]
 					{
 #if SLATE_STATS
