Moved the "Widget Caching" toggle to the top row of toggle buttons

diff --git a/Engine/Source/Developer/SlateReflector/Private/Widgets/SWidgetReflector.cpp b/Engine/Source/Developer/SlateReflector/Private/Widgets/SWidgetReflector.cpp
--- a/Engine/Source/Developer/SlateReflector/Private/Widgets/SWidgetReflector.cpp
+++ b/Engine/Source/Developer/SlateReflector/Private/Widgets/SWidgetReflector.cpp
@@ -437,6 +437,32 @@ void SWidgetReflector::Construct( const FArguments& InArgs )
 					]
 				]
 
+				+ SHorizontalBox::Slot()
+				.AutoWidth()
+				.Padding(FMargin(5.0f, 0.0f))
+				[
+					SNew(SCheckBox)
+					.Style(FCoreStyle::Get(), "ToggleButtonCheckbox")
+					.IsChecked_Lambda([&]()
+					{
+						return SInvalidationPanel::GetEnableWidgetCaching() ? ECheckBoxState::Checked : ECheckBoxState::Unchecked;
+					})
+					.OnCheckStateChanged_Lambda([&](const ECheckBoxState NewState)
+					{
+						SInvalidationPanel::SetEnableWidgetCaching(( NewState == ECheckBoxState::Checked ) ? true : false);
+					})
+					[
+						SNew(SBox)
+						.VAlign(VAlign_Center)
+						.HAlign(HAlign_Center)
+						.Padding(FMargin(4.0, 2.0))
+						[
+							SNew(STextBlock)
+							.Text(LOCTEXT("EnableWidgetCaching", "Widget Caching"))
+						]
+					]
+				]
+
 #if !(UE_BUILD_SHIPPING || UE_BUILD_TEST)
 				+SHorizontalBox::Slot()
 				.AutoWidth()
@@ -617,30 +643,6 @@ TSharedRef<SDockTab> SWidgetReflector::SpawnWidgetHierarchyTab(const FSpawnTabAr
 						.Text(LOCTEXT("ShowFocus", "Show Focus"))
 					]
                 ]
-				+ SHorizontalBox::Slot()
-				.AutoWidth()
-				.Padding(5.0f)
-				[
-					SNew(SCheckBox)
-					.Style(FCoreStyle::Get(), "ToggleButtonCheckbox")
-					.IsChecked_Lambda([&]()
-					{
-						return SInvalidationPanel::GetEnableWidgetCaching() ? ECheckBoxState::Checked : ECheckBoxState::Unchecked;
-					})
-					.OnCheckStateChanged_Lambda([&](const ECheckBoxState NewState)
-					{
-						SInvalidationPanel::SetEnableWidgetCaching(( NewState == ECheckBoxState::Checked ) ? true : false);
-					})
-					[
-						SNew(SBox)
-						.VAlign(VAlign_Center)
-						.HAlign(HAlign_Center)
-						[
-							SNew(STextBlock)
-							.Text(LOCTEXT("EnableWidgetCaching", "Widget Caching"))
-						]
-					]
-				]
 				
 				+SHorizontalBox::Slot()
 				.AutoWidth()
