Added support for taking widget snapshots from remote targets (Epic Friday)

This utilizes the message bus to allow you to take widget snapshots from any remote target that hosts the WidgetSnapshotService.

diff --git a/Engine/Source/Developer/SlateReflector/Private/SlateReflectorModule.cpp b/Engine/Source/Developer/SlateReflector/Private/SlateReflectorModule.cpp
--- a/Engine/Source/Developer/SlateReflector/Private/SlateReflectorModule.cpp
+++ b/Engine/Source/Developer/SlateReflector/Private/SlateReflectorModule.cpp
@@ -5,6 +5,7 @@
 #include "SDockTab.h"
 #include "ModuleManager.h"
 #include "SAtlasVisualizer.h"
+#include "WidgetSnapshotService.h"
 
 
 #define LOCTEXT_NAMESPACE "FSlateReflectorModule"
@@ -27,7 +28,8 @@ public:
 		if (!WidgetReflector.IsValid())
 		{
 			WidgetReflector = SNew(SWidgetReflector)
-				.ParentTab(InParentTab);
+				.ParentTab(InParentTab)
+				.WidgetSnapshotService(WidgetSnapshotService);
 			WidgetReflectorPtr = WidgetReflector;
 			FSlateApplication::Get().SetWidgetReflector(WidgetReflector.ToSharedRef());
 		}
@@ -160,6 +162,8 @@ public:
 
 	virtual void StartupModule() override
 	{
+		WidgetSnapshotService = MakeShareable(new FWidgetSnapshotService());
+
 		bHasRegisteredTabSpawners = false;
 		RegisterTabSpawner(nullptr);
 	}
@@ -167,6 +171,8 @@ public:
 	virtual void ShutdownModule() override
 	{
 		UnregisterTabSpawner();
+
+		WidgetSnapshotService.Reset();
 	}
 
 private:
@@ -204,6 +210,9 @@ private:
 
 	/** Holds the widget reflector singleton. */
 	TWeakPtr<SWidgetReflector> WidgetReflectorPtr;
+
+	/** The service for handling remote widget snapshots */
+	TSharedPtr<FWidgetSnapshotService> WidgetSnapshotService;
 };
 
 
diff --git a/Engine/Source/Developer/SlateReflector/Private/WidgetSnapshotMessages.h b/Engine/Source/Developer/SlateReflector/Private/WidgetSnapshotMessages.h
new file mode 100644
--- /dev/null
+++ b/Engine/Source/Developer/SlateReflector/Private/WidgetSnapshotMessages.h
@@ -0,0 +1,41 @@
+// Copyright 1998-2015 Epic Games, Inc. All Rights Reserved.
+
+#pragma once
+
+#include "WidgetSnapshotMessages.generated.h"
+
+
+/**
+ * Implements a message that is used to request a widget snapshot from a remote target.
+ */
+USTRUCT()
+struct FWidgetSnapshotRequest
+{
+	GENERATED_USTRUCT_BODY()
+
+	/** The instance ID of the remote target we want to receive a snapshot from */
+	UPROPERTY()
+	FGuid TargetInstanceId;
+
+	/** The request ID of this snapshot (used to identify the correct response from the target) */
+	UPROPERTY()
+	FGuid SnapshotRequestId;
+};
+
+
+/**
+ * Implements a message that is used to receive a widget snapshot from a remote target.
+ */
+USTRUCT()
+struct FWidgetSnapshotResponse
+{
+	GENERATED_USTRUCT_BODY()
+
+	/** The request ID of this snapshot (used to identify the correct response from the target) */
+	UPROPERTY()
+	FGuid SnapshotRequestId;
+
+	/** The snapshot data, to be used by FWidgetSnapshotData::LoadSnapshotFromBuffer */
+	UPROPERTY()
+	TArray<uint8> SnapshotData;
+};
diff --git a/Engine/Source/Developer/SlateReflector/Private/WidgetSnapshotService.cpp b/Engine/Source/Developer/SlateReflector/Private/WidgetSnapshotService.cpp
new file mode 100644
--- /dev/null
+++ b/Engine/Source/Developer/SlateReflector/Private/WidgetSnapshotService.cpp
@@ -0,0 +1,73 @@
+// Copyright 1998-2015 Epic Games, Inc. All Rights Reserved.
+
+#include "SlateReflectorPrivatePCH.h"
+#include "WidgetSnapshotService.h"
+#include "IMessagingModule.h"
+#include "MessageEndpointBuilder.h"
+#include "SWidgetSnapshotVisualizer.h"
+
+FWidgetSnapshotService::FWidgetSnapshotService()
+{
+	if (FPlatformMisc::SupportsMessaging() && FPlatformProcess::SupportsMultithreading())
+	{
+		IMessageBusPtr MessageBusPtr = IMessagingModule::Get().GetDefaultBus();
+
+		MessageEndpoint = FMessageEndpoint::Builder("FWidgetSnapshotService", MessageBusPtr.ToSharedRef())
+			.ReceivingOnThread(ENamedThreads::GameThread)
+			.Handling<FWidgetSnapshotRequest>(this, &FWidgetSnapshotService::HandleWidgetSnapshotRequestMessage)
+			.Handling<FWidgetSnapshotResponse>(this, &FWidgetSnapshotService::HandleWidgetSnapshotResponseMessage);
+
+		if (MessageEndpoint.IsValid())
+		{
+			MessageEndpoint->Subscribe<FWidgetSnapshotRequest>();
+		}
+	}
+}
+
+FGuid FWidgetSnapshotService::RequestSnapshot(const FGuid& InRemoteInstanceId, const FOnWidgetSnapshotResponse& OnResponse)
+{
+	if (MessageEndpoint.IsValid())
+	{
+		FWidgetSnapshotRequest* WidgetSnapshotRequest = new FWidgetSnapshotRequest;
+		WidgetSnapshotRequest->TargetInstanceId = InRemoteInstanceId;
+		WidgetSnapshotRequest->SnapshotRequestId = FGuid::NewGuid();
+
+		PendingSnapshotResponseHandlers.Add(WidgetSnapshotRequest->SnapshotRequestId, OnResponse);
+
+		MessageEndpoint->Publish(WidgetSnapshotRequest);
+
+		return WidgetSnapshotRequest->SnapshotRequestId;
+	}
+
+	return FGuid();
+}
+
+void FWidgetSnapshotService::AbortSnapshotRequest(const FGuid& InSnapshotRequestId)
+{
+	PendingSnapshotResponseHandlers.Remove(InSnapshotRequestId);
+}
+
+void FWidgetSnapshotService::HandleWidgetSnapshotRequestMessage(const FWidgetSnapshotRequest& Message, const IMessageContextRef& Context)
+{
+	if (MessageEndpoint.IsValid() && Message.TargetInstanceId == FApp::GetInstanceId())
+	{
+		FWidgetSnapshotData SnapshotData;
+		SnapshotData.TakeSnapshot();
+
+		FWidgetSnapshotResponse* WidgetSnapshotResponse = new FWidgetSnapshotResponse;
+		WidgetSnapshotResponse->SnapshotRequestId = Message.SnapshotRequestId;
+		SnapshotData.SaveSnapshotToBuffer(WidgetSnapshotResponse->SnapshotData);
+
+		MessageEndpoint->Send(WidgetSnapshotResponse, Context->GetSender());
+	}
+}
+
+void FWidgetSnapshotService::HandleWidgetSnapshotResponseMessage(const FWidgetSnapshotResponse& Message, const IMessageContextRef& Context)
+{
+	const FOnWidgetSnapshotResponse* FoundResponseHandler = PendingSnapshotResponseHandlers.Find(Message.SnapshotRequestId);
+	if (FoundResponseHandler && FoundResponseHandler->IsBound())
+	{
+		FoundResponseHandler->Execute(Message.SnapshotData);
+		PendingSnapshotResponseHandlers.Remove(Message.SnapshotRequestId);
+	}
+}
diff --git a/Engine/Source/Developer/SlateReflector/Private/WidgetSnapshotService.h b/Engine/Source/Developer/SlateReflector/Private/WidgetSnapshotService.h
new file mode 100644
--- /dev/null
+++ b/Engine/Source/Developer/SlateReflector/Private/WidgetSnapshotService.h
@@ -0,0 +1,36 @@
+// Copyright 1998-2015 Epic Games, Inc. All Rights Reserved.
+
+#pragma once
+
+#include "WidgetSnapshotMessages.h"
+#include "MessageEndpoint.h"
+
+/**
+ * Implements the service for handling remote widget snapshots.
+ */
+class FWidgetSnapshotService
+{
+public:
+	DECLARE_DELEGATE_OneParam(FOnWidgetSnapshotResponse, const TArray<uint8>& /*InSnapshotData*/);
+
+	FWidgetSnapshotService();
+
+	/** Request a snapshot from the given instance. The given delegate will be called when the response comes in */
+	FGuid RequestSnapshot(const FGuid& InRemoteInstanceId, const FOnWidgetSnapshotResponse& OnResponse);
+
+	/** Abort a request using the GUID previously obtained via a call to RequestSnapshot */
+	void AbortSnapshotRequest(const FGuid& InSnapshotRequestId);
+
+private:
+	/** Handles FWidgetSnapshotRequest messages */
+	void HandleWidgetSnapshotRequestMessage(const FWidgetSnapshotRequest& Message, const IMessageContextRef& Context);
+
+	/** Handles FWidgetSnapshotResponse messages */
+	void HandleWidgetSnapshotResponseMessage(const FWidgetSnapshotResponse& Message, const IMessageContextRef& Context);
+
+	/** Holds the message endpoint */
+	FMessageEndpointPtr MessageEndpoint;
+
+	/** Handlers awaiting their resultant snapshot data */
+	TMap<FGuid, FOnWidgetSnapshotResponse> PendingSnapshotResponseHandlers;
+};
diff --git a/Engine/Source/Developer/SlateReflector/Private/Widgets/SWidgetReflector.cpp b/Engine/Source/Developer/SlateReflector/Private/Widgets/SWidgetReflector.cpp
--- a/Engine/Source/Developer/SlateReflector/Private/Widgets/SWidgetReflector.cpp
+++ b/Engine/Source/Developer/SlateReflector/Private/Widgets/SWidgetReflector.cpp
@@ -5,12 +5,22 @@
 #include "SlateStats.h"
 #include "SInvalidationPanel.h"
 #include "SDockTab.h"
+#include "SNotificationList.h"
+#include "NotificationManager.h"
 #include "SWidgetSnapshotVisualizer.h"
+#include "WidgetSnapshotService.h"
 
 #if SLATE_REFLECTOR_HAS_DESKTOP_PLATFORM
 #include "DesktopPlatformModule.h"
 #endif // SLATE_REFLECTOR_HAS_DESKTOP_PLATFORM
 
+#if SLATE_REFLECTOR_HAS_SESSION_SERVICES
+#include "ISessionInfo.h"
+#include "ISessionInstanceInfo.h"
+#include "ISessionManager.h"
+#include "ISessionServicesModule.h"
+#endif // SLATE_REFLECTOR_HAS_SESSION_SERVICES
+
 #define LOCTEXT_NAMESPACE "SWidgetReflector"
 #define WITH_EVENT_LOGGING 0
 
@@ -53,6 +63,16 @@ struct FLoggedEvent
 namespace WidgetReflectorImpl
 {
 
+/** Information about a potential widget snapshot target */
+struct FWidgetSnapshotTarget
+{
+	/** Display name of the target (used in the UI) */
+	FText DisplayName;
+
+	/** Instance ID of the target */
+	FGuid InstanceId;
+};
+
 /** Different UI modes the widget reflector can be in */
 enum class EWidgetReflectorUIMode : uint8
 {
@@ -233,14 +253,44 @@ private:
 	/** Callback for getting the text of the pick button. */
 	FText HandlePickButtonText() const;
 
+	/** Callback to see whether the "Snapshot Target" combo should be enabled */
+	bool IsSnapshotTargetComboEnabled() const;
+
+	/** Callback to see whether the "Take Snapshot" button should be enabled */
+	bool IsTakeSnapshotButtonEnabled() const;
+
 	/** Callback for clicking the "Take Snapshot" button. */
 	FReply HandleTakeSnapshotButtonClicked();
 
+	/** Used as a callback for the "snapshot pending" notification item buttons, called when we should give up on a snapshot request */
+	void OnCancelPendingRemoteSnapshot();
+
+	/** Callback for when a remote widget snapshot is available */
+	void HandleRemoteSnapshotReceived(const TArray<uint8>& InSnapshotData);
+
 #if SLATE_REFLECTOR_HAS_DESKTOP_PLATFORM
 	/** Callback for clicking the "Load Snapshot" button. */
 	FReply HandleLoadSnapshotButtonClicked();
 #endif // SLATE_REFLECTOR_HAS_DESKTOP_PLATFORM
 
+	/** Called to update the list of available snapshot targets */
+	void UpdateAvailableSnapshotTargets();
+
+	/** Called to update the currently selected snapshot target (after the list has been refreshed) */
+	void UpdateSelectedSnapshotTarget();
+
+	/** Called when the list of available snapshot targets changes */
+	void OnAvailableSnapshotTargetsChanged();
+
+	/** Get the display name of the currently selected snapshot target */
+	FText GetSelectedSnapshotTargetDisplayName() const;
+
+	/** Generate a row widget for the available targets combo box */
+	TSharedRef<SWidget> HandleGenerateAvailableSnapshotComboItemWidget(TSharedPtr<FWidgetSnapshotTarget> InItem) const;
+	
+	/** Update the selected target when the combo box selection is changed */
+	void HandleAvailableSnapshotComboSelectionChanged(TSharedPtr<FWidgetSnapshotTarget> InItem, ESelectInfo::Type InSeletionInfo);
+
 	/** Callback for generating a row in the reflector tree view. */
 	TSharedRef<ITableRow> HandleReflectorTreeGenerateRow( TSharedRef<FWidgetReflectorNodeBase> InReflectorNode, const TSharedRef<STableViewBase>& OwnerTable );
 
@@ -272,6 +322,14 @@ private:
 	FWidgetSnapshotData SnapshotData;
 	TSharedPtr<SWidgetSnapshotVisualizer> WidgetSnapshotVisualizer;
 
+	/** List of available snapshot targets, as well as the one we currently have selected */
+	TSharedPtr<SComboBox<TSharedPtr<FWidgetSnapshotTarget>>> AvailableSnapshotTargetsComboBox;
+	TArray<TSharedPtr<FWidgetSnapshotTarget>> AvailableSnapshotTargets;
+	FGuid SelectedSnapshotTargetInstanceId;
+	TSharedPtr<FWidgetSnapshotService> WidgetSnapshotService;
+	TWeakPtr<SNotificationItem> WidgetSnapshotNotificationPtr;
+	FGuid RemoteSnapshotRequestId;
+
 	SSplitter::FSlot* WidgetInfoLocation;
 
 	FAccessSourceCode SourceAccessDelegate;
@@ -319,6 +377,17 @@ void SWidgetReflector::Construct( const FArguments& InArgs )
 	LastMouseClickTime = -1;
 	CursorPingPosition = FVector2D::ZeroVector;
 
+	WidgetSnapshotService = InArgs._WidgetSnapshotService;
+
+#if SLATE_REFLECTOR_HAS_SESSION_SERVICES
+	{
+		ISessionManagerRef SessionManager = FModuleManager::LoadModuleChecked<ISessionServicesModule>("SessionServices").GetSessionManager();
+		SessionManager->OnSessionsUpdated().AddSP(this, &SWidgetReflector::OnAvailableSnapshotTargetsChanged);
+	}
+#endif // SLATE_REFLECTOR_HAS_SESSION_SERVICES
+	SelectedSnapshotTargetInstanceId = FApp::GetInstanceId();
+	UpdateAvailableSnapshotTargets();
+
 #if SLATE_STATS
 	const FName TabLayoutName = "WidgetReflector_Layout_v1";
 #else
@@ -618,7 +687,7 @@ void SWidgetReflector::Construct( const FArguments& InArgs )
 
 TSharedRef<SDockTab> SWidgetReflector::SpawnWidgetHierarchyTab(const FSpawnTabArgs& Args)
 {
-	return SNew(SDockTab)
+	TSharedRef<SDockTab> SpawnedTab = SNew(SDockTab)
 		.Label(LOCTEXT("WidgetHierarchyTab", "Widget Hierarchy"))
 		//.OnCanCloseTab_Lambda([]() { return false; }) // Can't prevent this as it stops the editor from being able to close while the widget reflector is open
 		[
@@ -658,16 +727,44 @@ TSharedRef<SDockTab> SWidgetReflector::SpawnWidgetHierarchyTab(const FSpawnTabAr
 					]
 				]
 
+				+SHorizontalBox::Slot()
+				[
+					SNew(SSpacer)
+				]
+
 				+SHorizontalBox::Slot()
 				.AutoWidth()
 				.Padding(FMargin(5.0f, 0.0f))
 				[
-					// Button that controls taking a snapshot of the current window(s)
-					SNew(SButton)
-					.OnClicked(this, &SWidgetReflector::HandleTakeSnapshotButtonClicked)
+					SNew(SHorizontalBox)
+
+					+SHorizontalBox::Slot()
+					.AutoWidth()
 					[
-						SNew(STextBlock)
-						.Text(LOCTEXT("TakeSnapshotButtonText", "Take Snapshot"))
+						// Button that controls taking a snapshot of the current window(s)
+						SNew(SButton)
+						.IsEnabled(this, &SWidgetReflector::IsTakeSnapshotButtonEnabled)
+						.OnClicked(this, &SWidgetReflector::HandleTakeSnapshotButtonClicked)
+						[
+							SNew(STextBlock)
+							.Text(LOCTEXT("TakeSnapshotButtonText", "Take Snapshot"))
+						]
+					]
+
+					+SHorizontalBox::Slot()
+					.AutoWidth()
+					[
+						// Button that controls the target for the snapshot operation
+						SAssignNew(AvailableSnapshotTargetsComboBox, SComboBox<TSharedPtr<FWidgetSnapshotTarget>>)
+						.IsEnabled(this, &SWidgetReflector::IsSnapshotTargetComboEnabled)
+						.ToolTipText(LOCTEXT("ChooseSnapshotTargetToolTipText", "Choose Snapshot Target"))
+						.OptionsSource(&AvailableSnapshotTargets)
+						.OnGenerateWidget(this, &SWidgetReflector::HandleGenerateAvailableSnapshotComboItemWidget)
+						.OnSelectionChanged(this, &SWidgetReflector::HandleAvailableSnapshotComboSelectionChanged)
+						[
+							SNew(STextBlock)
+							.Text(this, &SWidgetReflector::GetSelectedSnapshotTargetDisplayName)
+						]
 					]
 				]
 
@@ -734,6 +831,10 @@ TSharedRef<SDockTab> SWidgetReflector::SpawnWidgetHierarchyTab(const FSpawnTabAr
 				]
 			]
 		];
+
+	UpdateSelectedSnapshotTarget();
+
+	return SpawnedTab;
 }
 
 
@@ -1246,20 +1347,123 @@ FText SWidgetReflector::HandlePickButtonText() const
 }
 
 
+bool SWidgetReflector::IsSnapshotTargetComboEnabled() const
+{
+	return SLATE_REFLECTOR_HAS_SESSION_SERVICES != 0 && !RemoteSnapshotRequestId.IsValid();
+}
+
+
+bool SWidgetReflector::IsTakeSnapshotButtonEnabled() const
+{
+	return SelectedSnapshotTargetInstanceId.IsValid() && !RemoteSnapshotRequestId.IsValid();
+}
+
+
 FReply SWidgetReflector::HandleTakeSnapshotButtonClicked()
 {
+	// Local snapshot?
+	if (SelectedSnapshotTargetInstanceId == FApp::GetInstanceId())
+	{
+		SetUIMode(EWidgetReflectorUIMode::Snapshot);
+
+		// Take a snapshot of any window(s) that are currently open
+		SnapshotData.TakeSnapshot();
+
+		// Rebuild the reflector tree from the snapshot data
+		ReflectorTreeRoot = SnapshotData.GetWindowsRef();
+		ReflectorTree->RequestTreeRefresh();
+
+		WidgetSnapshotVisualizer->SnapshotDataUpdated();
+	}
+	else
+	{
+		// Remote snapshot - these can take a while, show a progress message
+		FNotificationInfo Info(LOCTEXT("RemoteWidgetSnapshotPendingNotificationText", "Waiting for Remote Widget Snapshot Data"));
+
+		// Add the buttons with text, tooltip and callback
+		Info.ButtonDetails.Add(FNotificationButtonInfo(
+			LOCTEXT("CancelPendingSnapshotButtonText", "Cancel"), 
+			LOCTEXT("CancelPendingSnapshotButtonToolTipText", "Cancel the pending widget snapshot request."), 
+			FSimpleDelegate::CreateSP(this, &SWidgetReflector::OnCancelPendingRemoteSnapshot)
+			));
+
+		// We will be keeping track of this ourselves
+		Info.bFireAndForget = false;
+
+		// Launch notification
+		WidgetSnapshotNotificationPtr = FSlateNotificationManager::Get().AddNotification(Info);
+
+		if (WidgetSnapshotNotificationPtr.IsValid())
+		{
+			WidgetSnapshotNotificationPtr.Pin()->SetCompletionState(SNotificationItem::CS_Pending);
+		}
+
+		RemoteSnapshotRequestId = WidgetSnapshotService->RequestSnapshot(SelectedSnapshotTargetInstanceId, FWidgetSnapshotService::FOnWidgetSnapshotResponse::CreateSP(this, &SWidgetReflector::HandleRemoteSnapshotReceived));
+
+		if (!RemoteSnapshotRequestId.IsValid())
+		{
+			TSharedPtr<SNotificationItem> WidgetSnapshotNotificationPin = WidgetSnapshotNotificationPtr.Pin();
+
+			if (WidgetSnapshotNotificationPin.IsValid())
+			{
+				WidgetSnapshotNotificationPin->SetText(LOCTEXT("RemoteWidgetSnapshotFailedNotificationText", "Remote Widget Snapshot Failed"));
+				WidgetSnapshotNotificationPin->SetCompletionState(SNotificationItem::CS_Fail);
+				WidgetSnapshotNotificationPin->ExpireAndFadeout();
+
+				WidgetSnapshotNotificationPtr.Reset();
+			}
+		}
+	}
+
+	return FReply::Handled();
+}
+
+
+void SWidgetReflector::OnCancelPendingRemoteSnapshot()
+{
+	TSharedPtr<SNotificationItem> WidgetSnapshotNotificationPin = WidgetSnapshotNotificationPtr.Pin();
+
+	if (WidgetSnapshotNotificationPin.IsValid())
+	{
+		WidgetSnapshotNotificationPin->SetText(LOCTEXT("RemoteWidgetSnapshotAbortedNotificationText", "Aborted Remote Widget Snapshot"));
+		WidgetSnapshotNotificationPin->SetCompletionState(SNotificationItem::CS_Fail);
+		WidgetSnapshotNotificationPin->ExpireAndFadeout();
+
+		WidgetSnapshotNotificationPtr.Reset();
+	}
+
+	WidgetSnapshotService->AbortSnapshotRequest(RemoteSnapshotRequestId);
+	RemoteSnapshotRequestId = FGuid();
+}
+
+
+void SWidgetReflector::HandleRemoteSnapshotReceived(const TArray<uint8>& InSnapshotData)
+{
+	{
+		TSharedPtr<SNotificationItem> WidgetSnapshotNotificationPin = WidgetSnapshotNotificationPtr.Pin();
+
+		if (WidgetSnapshotNotificationPin.IsValid())
+		{
+			WidgetSnapshotNotificationPin->SetText(LOCTEXT("RemoteWidgetSnapshotReceivedNotificationText", "Remote Widget Snapshot Data Received"));
+			WidgetSnapshotNotificationPin->SetCompletionState(SNotificationItem::CS_Success);
+			WidgetSnapshotNotificationPin->ExpireAndFadeout();
+
+			WidgetSnapshotNotificationPtr.Reset();
+		}
+	}
+
+	RemoteSnapshotRequestId = FGuid();
+
 	SetUIMode(EWidgetReflectorUIMode::Snapshot);
 
-	// Take a snapshot of any window(s) that are currently open
-	SnapshotData.TakeSnapshot();
+	// Load up the remote data
+	SnapshotData.LoadSnapshotFromBuffer(InSnapshotData);
 
 	// Rebuild the reflector tree from the snapshot data
 	ReflectorTreeRoot = SnapshotData.GetWindowsRef();
 	ReflectorTree->RequestTreeRefresh();
 
 	WidgetSnapshotVisualizer->SnapshotDataUpdated();
-
-	return FReply::Handled();
 }
 
 
@@ -1302,6 +1506,120 @@ FReply SWidgetReflector::HandleLoadSnapshotButtonClicked()
 #endif // SLATE_REFLECTOR_HAS_DESKTOP_PLATFORM
 
 
+void SWidgetReflector::UpdateAvailableSnapshotTargets()
+{
+	AvailableSnapshotTargets.Reset();
+
+#if SLATE_REFLECTOR_HAS_SESSION_SERVICES
+	{
+		ISessionManagerRef SessionManager = FModuleManager::LoadModuleChecked<ISessionServicesModule>("SessionServices").GetSessionManager();
+
+		TArray<TSharedPtr<ISessionInfo>> AvailableSessions;
+		SessionManager->GetSessions(AvailableSessions);
+
+		for (const auto& AvailableSession : AvailableSessions)
+		{
+			// Only allow sessions belonging to the current user
+			if (AvailableSession->GetSessionOwner() != FApp::GetSessionOwner())
+			{
+				continue;
+			}
+
+			TArray<TSharedPtr<ISessionInstanceInfo>> AvailableInstances;
+			AvailableSession->GetInstances(AvailableInstances);
+
+			for (const auto& AvailableInstance : AvailableInstances)
+			{
+				FWidgetSnapshotTarget SnapshotTarget;
+				SnapshotTarget.DisplayName = FText::Format(LOCTEXT("SnapshotTargetDisplayNameFmt", "{0} ({1})"), FText::FromString(AvailableInstance->GetInstanceName()), FText::FromString(AvailableInstance->GetPlatformName()));
+				SnapshotTarget.InstanceId = AvailableInstance->GetInstanceId();
+
+				AvailableSnapshotTargets.Add(MakeShareable(new FWidgetSnapshotTarget(SnapshotTarget)));
+			}
+		}
+	}
+#else
+	{
+		// No session services, just add an entry that lets us snapshot ourself
+		FWidgetSnapshotTarget SnapshotTarget;
+		SnapshotTarget.DisplayName = FText::FromString(FApp::GetInstanceName());
+		SnapshotTarget.InstanceId = FApp::GetInstanceId();
+
+		AvailableSnapshotTargets.Add(MakeShareable(new FWidgetSnapshotTarget(SnapshotTarget)));
+	}
+#endif
+}
+
+
+void SWidgetReflector::UpdateSelectedSnapshotTarget()
+{
+	if (AvailableSnapshotTargetsComboBox.IsValid())
+	{
+		const TSharedPtr<FWidgetSnapshotTarget>* FoundSnapshotTarget = AvailableSnapshotTargets.FindByPredicate([this](const TSharedPtr<FWidgetSnapshotTarget>& InAvailableSnapshotTarget) -> bool
+		{
+			return InAvailableSnapshotTarget->InstanceId == SelectedSnapshotTargetInstanceId;
+		});
+
+		if (FoundSnapshotTarget)
+		{
+			AvailableSnapshotTargetsComboBox->SetSelectedItem(*FoundSnapshotTarget);
+		}
+		else if (AvailableSnapshotTargets.Num() > 0)
+		{
+			SelectedSnapshotTargetInstanceId = AvailableSnapshotTargets[0]->InstanceId;
+			AvailableSnapshotTargetsComboBox->SetSelectedItem(AvailableSnapshotTargets[0]);
+		}
+		else
+		{
+			SelectedSnapshotTargetInstanceId = FGuid();
+			AvailableSnapshotTargetsComboBox->SetSelectedItem(nullptr);
+		}
+	}
+}
+
+
+void SWidgetReflector::OnAvailableSnapshotTargetsChanged()
+{
+	UpdateAvailableSnapshotTargets();
+	UpdateSelectedSnapshotTarget();
+}
+
+
+FText SWidgetReflector::GetSelectedSnapshotTargetDisplayName() const
+{
+	if (AvailableSnapshotTargetsComboBox.IsValid())
+	{
+		TSharedPtr<FWidgetSnapshotTarget> SelectedSnapshotTarget = AvailableSnapshotTargetsComboBox->GetSelectedItem();
+		if (SelectedSnapshotTarget.IsValid())
+		{
+			return SelectedSnapshotTarget->DisplayName;
+		}
+	}
+
+	return FText::GetEmpty();
+}
+
+
+TSharedRef<SWidget> SWidgetReflector::HandleGenerateAvailableSnapshotComboItemWidget(TSharedPtr<FWidgetSnapshotTarget> InItem) const
+{
+	return SNew(STextBlock)
+		.Text(InItem->DisplayName);
+}
+
+
+void SWidgetReflector::HandleAvailableSnapshotComboSelectionChanged(TSharedPtr<FWidgetSnapshotTarget> InItem, ESelectInfo::Type InSeletionInfo)
+{
+	if (InItem.IsValid())
+	{
+		SelectedSnapshotTargetInstanceId = InItem->InstanceId;
+	}
+	else
+	{
+		SelectedSnapshotTargetInstanceId = FGuid();
+	}
+}
+
+
 TSharedRef<ITableRow> SWidgetReflector::HandleReflectorTreeGenerateRow( TSharedRef<FWidgetReflectorNodeBase> InReflectorNode, const TSharedRef<STableViewBase>& OwnerTable )
 {
 	return SNew(SReflectorTreeWidgetItem, OwnerTable)
diff --git a/Engine/Source/Developer/SlateReflector/Private/Widgets/SWidgetReflector.h b/Engine/Source/Developer/SlateReflector/Private/Widgets/SWidgetReflector.h
--- a/Engine/Source/Developer/SlateReflector/Private/Widgets/SWidgetReflector.h
+++ b/Engine/Source/Developer/SlateReflector/Private/Widgets/SWidgetReflector.h
@@ -4,6 +4,8 @@
 
 #include "IWidgetReflector.h"
 
+class FWidgetSnapshotService;
+
 /**
  * Widget reflector implementation.
  * User widget to enable iteration without recompilation.
@@ -18,6 +20,7 @@ public:
 	{ }
 		
 		SLATE_ARGUMENT(TSharedPtr<SDockTab>, ParentTab)
+		SLATE_ARGUMENT(TSharedPtr<FWidgetSnapshotService>, WidgetSnapshotService)
 
 	SLATE_END_ARGS()
 
diff --git a/Engine/Source/Developer/SlateReflector/SlateReflector.Build.cs b/Engine/Source/Developer/SlateReflector/SlateReflector.Build.cs
--- a/Engine/Source/Developer/SlateReflector/SlateReflector.Build.cs
+++ b/Engine/Source/Developer/SlateReflector/SlateReflector.Build.cs
@@ -22,6 +22,18 @@ public class SlateReflector : ModuleRules
 			}
 		);
 
+		PrivateIncludePathModuleNames.AddRange(
+			new string[] {
+				"Messaging",
+			}
+		);
+
+		DynamicallyLoadedModuleNames.AddRange(
+			new string[] {
+				"Messaging",
+			}
+		);
+
 		PrivateIncludePaths.AddRange(
 			new string[] {
 				"Developer/SlateReflector/Private",
@@ -30,6 +42,28 @@ public class SlateReflector : ModuleRules
 			}
 		);
 
+		// Editor builds include SessionServices to populate the remote target drop-down for remote widget snapshots
+		if (Target.Type == TargetRules.TargetType.Editor)
+		{
+			Definitions.Add("SLATE_REFLECTOR_HAS_SESSION_SERVICES=1");
+
+			PrivateIncludePathModuleNames.AddRange(
+				new string[] {
+					"SessionServices",
+				}
+			);
+
+			DynamicallyLoadedModuleNames.AddRange(
+				new string[] {
+					"SessionServices",
+				}
+			);
+		}
+		else
+		{
+			Definitions.Add("SLATE_REFLECTOR_HAS_SESSION_SERVICES=0");
+		}
+
 		// DesktopPlatform is only available for Editor and Program targets (running on a desktop platform)
 		bool IsDesktopPlatformType = Target.Platform == UnrealBuildTool.UnrealTargetPlatform.Win32
 			|| Target.Platform == UnrealBuildTool.UnrealTargetPlatform.Win64
diff --git a/Engine/Source/Runtime/Launch/Private/LaunchEngineLoop.cpp b/Engine/Source/Runtime/Launch/Private/LaunchEngineLoop.cpp
--- a/Engine/Source/Runtime/Launch/Private/LaunchEngineLoop.cpp
+++ b/Engine/Source/Runtime/Launch/Private/LaunchEngineLoop.cpp
@@ -1939,6 +1939,11 @@ bool FEngineLoop::LoadStartupCoreModules()
 	if ( !IsRunningDedicatedServer() )
 	{
 		FModuleManager::Get().LoadModule("Slate");
+
+#if !UE_BUILD_SHIPPING
+		// Need to load up the SlateReflector module to initialize the WidgetSnapshotService
+		FModuleManager::Get().LoadModule("SlateReflector");
+#endif // !UE_BUILD_SHIPPING
 	}
 
 #if WITH_EDITOR
