Change address to be more portable

MS compiler does not place a '0x' on %p formating. Linux/Mac append a '0x' to the address

diff --git a/Engine/Source/Developer/SlateReflector/Private/Models/WidgetReflectorNode.cpp b/Engine/Source/Developer/SlateReflector/Private/Models/WidgetReflectorNode.cpp
--- a/Engine/Source/Developer/SlateReflector/Private/Models/WidgetReflectorNode.cpp
+++ b/Engine/Source/Developer/SlateReflector/Private/Models/WidgetReflectorNode.cpp
@@ -656,7 +656,7 @@ FVector2D FWidgetReflectorNodeUtils::GetWidgetDesiredSize(const TSharedPtr<SWidg
 
 FString FWidgetReflectorNodeUtils::GetWidgetAddress(const TSharedPtr<SWidget>& InWidget)
 {
-	return FString::Printf(TEXT("0x%p"), InWidget.Get());
+	return FString::Printf(TEXT("0x%0llx"), reinterpret_cast<uint64>(InWidget.Get()));
 }
 
 FSlateColor FWidgetReflectorNodeUtils::GetWidgetForegroundColor(const TSharedPtr<SWidget>& InWidget)
