Slate - Introducing a metadata you can attach to widgets to make it be ignored by the picker in Widget Reflector for paint hit testing.  The debug paint canvas can ruin the experience when you're trying to pick against painted widgets, so now it gets ignored.

Matt.Kuhlenschmidt

diff --git a/Engine/Source/Developer/SlateReflector/Private/VisualTreeCapture.cpp b/Engine/Source/Developer/SlateReflector/Private/VisualTreeCapture.cpp
--- a/Engine/Source/Developer/SlateReflector/Private/VisualTreeCapture.cpp
+++ b/Engine/Source/Developer/SlateReflector/Private/VisualTreeCapture.cpp
@@ -6,6 +6,7 @@
 #include "Rendering/DrawElements.h"
 #include "Widgets/SWidget.h"
 #include "Framework/Application/SlateApplication.h"
+#include "Types/InvisibleToWidgetReflectorMetaData.h"
 
 static float VectorSign(const FVector2D& Vec, const FVector2D& A, const FVector2D& B)
 {
@@ -185,6 +186,12 @@ void FVisualTreeCapture::ElementAdded(const FSlateWindowElementList& ElementList
 	{
 		if (Tree->WidgetStack.Num() > 0)
 		{
+			// Ignore any element added from a widget that's invisible to the widget reflector.
+			if (Tree->WidgetStack.Top().Pin()->GetMetaData<FInvisibleToWidgetReflectorMetaData>())
+			{
+				return;
+			}
+
 			FVisualEntry Entry(InElementIndex);
 			Entry.Widget = Tree->WidgetStack.Top();
 			Tree->Entries.Add(Entry);
diff --git a/Engine/Source/Runtime/Engine/Private/Slate/SGameLayerManager.cpp b/Engine/Source/Runtime/Engine/Private/Slate/SGameLayerManager.cpp
--- a/Engine/Source/Runtime/Engine/Private/Slate/SGameLayerManager.cpp
+++ b/Engine/Source/Runtime/Engine/Private/Slate/SGameLayerManager.cpp
@@ -16,6 +16,7 @@
 #include "Widgets/Layout/SPopup.h"
 #include "Widgets/Layout/SWindowTitleBarArea.h"
 #include "DebugCanvas.h"
+#include "Types/InvisibleToWidgetReflectorMetaData.h"
 
 /* SGameLayerManager interface
  *****************************************************************************/
@@ -86,7 +87,7 @@ void SGameLayerManager::Construct(const SGameLayerManager::FArguments& InArgs)
 				[
 					SAssignNew(DebugCanvas, SDebugCanvas)
 					.SceneViewport(InArgs._SceneViewport)
-
+					.AddMetaData<FInvisibleToWidgetReflectorMetaData>(FInvisibleToWidgetReflectorMetaData())
 				]
 			]
 		];
diff --git a/Engine/Source/Runtime/SlateCore/Public/Types/InvisibleToWidgetReflectorMetaData.h b/Engine/Source/Runtime/SlateCore/Public/Types/InvisibleToWidgetReflectorMetaData.h
new file mode 100644
--- /dev/null
+++ b/Engine/Source/Runtime/SlateCore/Public/Types/InvisibleToWidgetReflectorMetaData.h
@@ -0,0 +1,21 @@
+// Copyright 1998-2019 Epic Games, Inc. All Rights Reserved.
+
+#pragma once
+
+#include "CoreMinimal.h"
+#include "UObject/Object.h"
+#include "UObject/WeakObjectPtr.h"
+#include "Types/ISlateMetaData.h"
+
+/**
+ * When using the widget reflector, it may be necessary to make some widgets non-pickable, like the debug canvas.
+ */
+class FInvisibleToWidgetReflectorMetaData : public ISlateMetaData
+{
+public:
+	SLATE_METADATA_TYPE(FInvisibleToWidgetReflectorMetaData, ISlateMetaData)
+
+	FInvisibleToWidgetReflectorMetaData()
+	{
+	}
+};
