WidgetReflector - Make sure to reset the visual tree capture after switching to hittesting reflecting, otherwise, it gets stuck in visual picking mode, which we don't want.

diff --git a/Engine/Source/Developer/SlateReflector/Private/VisualTreeCapture.cpp b/Engine/Source/Developer/SlateReflector/Private/VisualTreeCapture.cpp
--- a/Engine/Source/Developer/SlateReflector/Private/VisualTreeCapture.cpp
+++ b/Engine/Source/Developer/SlateReflector/Private/VisualTreeCapture.cpp
@@ -120,6 +120,11 @@ void FVisualTreeCapture::Disable()
 #endif
 }
 
+void FVisualTreeCapture::Reset()
+{
+	VisualTrees.Reset();
+}
+
 TSharedPtr<FVisualTreeSnapshot> FVisualTreeCapture::GetVisualTreeForWindow(SWindow* InWindow)
 {
 	return VisualTrees.FindRef(InWindow);
diff --git a/Engine/Source/Developer/SlateReflector/Private/VisualTreeCapture.h b/Engine/Source/Developer/SlateReflector/Private/VisualTreeCapture.h
--- a/Engine/Source/Developer/SlateReflector/Private/VisualTreeCapture.h
+++ b/Engine/Source/Developer/SlateReflector/Private/VisualTreeCapture.h
@@ -52,6 +52,7 @@ public:
 
 	void Enable();
 	void Disable();
+	void Reset();
 
 	TSharedPtr<FVisualTreeSnapshot> GetVisualTreeForWindow(SWindow* InWindow);
 	
diff --git a/Engine/Source/Developer/SlateReflector/Private/Widgets/SWidgetReflector.cpp b/Engine/Source/Developer/SlateReflector/Private/Widgets/SWidgetReflector.cpp
--- a/Engine/Source/Developer/SlateReflector/Private/Widgets/SWidgetReflector.cpp
+++ b/Engine/Source/Developer/SlateReflector/Private/Widgets/SWidgetReflector.cpp
@@ -281,22 +281,24 @@ private:
 	{
 		if (PickingMode != InMode)
 		{
-			{
-				SInvalidationPanel::SetEnableWidgetCaching(true);
-
-				VisualCapture.Disable();
-			}
+			// Disable visual picking, and renable widget caching.
+			SInvalidationPanel::SetEnableWidgetCaching(true);
+			VisualCapture.Disable();
 
+			// Enable the picking mode.
 			PickingMode = InMode;
 
-			if (PickingMode == EWidgetPickingMode::Drawable || PickingMode == EWidgetPickingMode::HitTesting)
+			// If we're enabling hit test, reset the visual capture entirely, we don't want to use the visual tree.
+			if (PickingMode == EWidgetPickingMode::HitTesting)
 			{
+				VisualCapture.Reset();
 				SInvalidationPanel::SetEnableWidgetCaching(false);
 			}
-
-			if (PickingMode == EWidgetPickingMode::Drawable)
+			// If we're using the drawing picking mode enable it!
+			else if (PickingMode == EWidgetPickingMode::Drawable)
 			{
 				VisualCapture.Enable();
+				SInvalidationPanel::SetEnableWidgetCaching(false);
 			}
 		}
 	}
