Slate - Adding improvements to the widget reflector so that we can tell if a zero size widget is present, as it may have padding around it, so now a line will appear if it's zero sized, rather than nothing.

diff --git a/Engine/Source/Developer/SlateReflector/Private/Models/WidgetReflectorNode.cpp b/Engine/Source/Developer/SlateReflector/Private/Models/WidgetReflectorNode.cpp
--- a/Engine/Source/Developer/SlateReflector/Private/Models/WidgetReflectorNode.cpp
+++ b/Engine/Source/Developer/SlateReflector/Private/Models/WidgetReflectorNode.cpp
@@ -121,6 +121,11 @@ FText FLiveWidgetReflectorNode::GetWidgetVisibilityText() const
 	return FWidgetReflectorNodeUtils::GetWidgetVisibilityText(Widget.Pin());
 }
 
+bool FLiveWidgetReflectorNode::GetWidgetVisible() const
+{
+	return FWidgetReflectorNodeUtils::GetWidgetVisibility(Widget.Pin());
+}
+
 FText FLiveWidgetReflectorNode::GetWidgetClippingText() const
 {
 	return FWidgetReflectorNodeUtils::GetWidgetClippingText(Widget.Pin());
@@ -198,6 +203,7 @@ FSnapshotWidgetReflectorNode::FSnapshotWidgetReflectorNode(const FArrangedWidget
 	, CachedWidgetType(FWidgetReflectorNodeUtils::GetWidgetType(InArrangedWidget.Widget))
 	, CachedWidgetTypeAndShortName(FWidgetReflectorNodeUtils::GetWidgetTypeAndShortName(InArrangedWidget.Widget))
 	, CachedWidgetVisibilityText(FWidgetReflectorNodeUtils::GetWidgetVisibilityText(InArrangedWidget.Widget))
+	, bCachedWidgetVisible(FWidgetReflectorNodeUtils::GetWidgetVisibility(InArrangedWidget.Widget))
 	, bCachedWidgetFocusable(FWidgetReflectorNodeUtils::GetWidgetFocusable(InArrangedWidget.Widget))
 	, CachedWidgetClippingText(FWidgetReflectorNodeUtils::GetWidgetClippingText(InArrangedWidget.Widget))
 	, CachedWidgetReadableLocation(FWidgetReflectorNodeUtils::GetWidgetReadableLocation(InArrangedWidget.Widget))
@@ -241,6 +247,11 @@ bool FSnapshotWidgetReflectorNode::GetWidgetFocusable() const
 	return bCachedWidgetFocusable;
 }
 
+bool FSnapshotWidgetReflectorNode::GetWidgetVisible() const
+{
+	return bCachedWidgetFocusable;
+}
+
 FText FSnapshotWidgetReflectorNode::GetWidgetClippingText() const
 {
 	return CachedWidgetClippingText;
@@ -563,13 +574,28 @@ TSharedRef<FWidgetReflectorNodeBase> FWidgetReflectorNodeUtils::NewNodeTreeFrom(
 {
 	TSharedRef<FWidgetReflectorNodeBase> NewNodeInstance = NewNode(InNodeType, InWidgetGeometry);
 
-	FArrangedChildren ArrangedChildren(EVisibility::All);
-	InWidgetGeometry.Widget->ArrangeChildren(InWidgetGeometry.Geometry, ArrangedChildren);
-	
-	for (int32 WidgetIndex = 0; WidgetIndex < ArrangedChildren.Num(); ++WidgetIndex)
+	TSharedRef<SWidget> CurWidgetParent = InWidgetGeometry.Widget;
+	if (FChildren* Children = CurWidgetParent->GetChildren())
 	{
-		// Note that we include both visible and invisible children!
-		NewNodeInstance->AddChildNode(NewNodeTreeFrom(InNodeType, ArrangedChildren[WidgetIndex]));
+		for (int32 ChildIndex = 0; ChildIndex < Children->Num(); ++ChildIndex)
+		{
+			TSharedRef<SWidget> ChildWidget = Children->GetChildAt(ChildIndex);
+			FGeometry ChildGeometry = ChildWidget->GetCachedGeometry();
+			const EVisibility CurWidgetVisibility = ChildWidget->GetVisibility();
+
+			// Don't add geometry for completely collapsed stuff
+			if (CurWidgetVisibility == EVisibility::Collapsed)
+			{
+				ChildGeometry = FGeometry();
+			}
+			else if (!CurWidgetParent->ValidatePathToChild(&ChildWidget.Get()))
+			{
+				ChildGeometry = FGeometry();
+			}
+
+			// Note that we include both visible and invisible children!
+			NewNodeInstance->AddChildNode(NewNodeTreeFrom(InNodeType, FArrangedWidget(ChildWidget, ChildGeometry)));
+		}
 	}
 
 	return NewNodeInstance;
@@ -624,6 +650,11 @@ FText FWidgetReflectorNodeUtils::GetWidgetVisibilityText(const TSharedPtr<SWidge
 	return (InWidget.IsValid()) ? FText::FromString(InWidget->GetVisibility().ToString()) : FText::GetEmpty();
 }
 
+bool FWidgetReflectorNodeUtils::GetWidgetVisibility(const TSharedPtr<SWidget>& InWidget)
+{
+	return InWidget.IsValid() ? InWidget->GetVisibility().IsVisible() : false;
+}
+
 bool FWidgetReflectorNodeUtils::GetWidgetFocusable(const TSharedPtr<SWidget>& InWidget)
 {
 	return InWidget.IsValid() ? InWidget->SupportsKeyboardFocus() : false;
diff --git a/Engine/Source/Developer/SlateReflector/Private/Models/WidgetReflectorNode.h b/Engine/Source/Developer/SlateReflector/Private/Models/WidgetReflectorNode.h
--- a/Engine/Source/Developer/SlateReflector/Private/Models/WidgetReflectorNode.h
+++ b/Engine/Source/Developer/SlateReflector/Private/Models/WidgetReflectorNode.h
@@ -79,6 +79,9 @@ public:
 	*/
 	virtual bool GetWidgetFocusable() const = 0;
 
+	/** Is the widget visible? */
+	virtual bool GetWidgetVisible() const = 0;
+
 	/**
 	 * The human readable location for widgets that are defined in C++ is the file and line number
 	 * The human readable location for widgets that are defined in UMG is the asset name
@@ -218,6 +221,7 @@ public:
 	virtual FText GetWidgetType() const override;
 	virtual FText GetWidgetTypeAndShortName() const override;
 	virtual FText GetWidgetVisibilityText() const override;
+	virtual bool GetWidgetVisible() const override;
 	virtual FText GetWidgetClippingText() const override;
 	virtual bool GetWidgetFocusable() const override;
 	virtual FText GetWidgetReadableLocation() const override;
@@ -269,6 +273,7 @@ public:
 	virtual FText GetWidgetTypeAndShortName() const override;
 	virtual FText GetWidgetVisibilityText() const override;
 	virtual FText GetWidgetClippingText() const override;
+	virtual bool GetWidgetVisible() const override;
 	virtual bool GetWidgetFocusable() const override;
 	virtual FText GetWidgetReadableLocation() const override;
 	virtual FString GetWidgetFile() const override;
@@ -307,6 +312,9 @@ private:
 	/** The visibility string of the widget at the point it was passed to Initialize */
 	FText CachedWidgetVisibilityText;
 
+	/** Is the widget visible? */
+	bool bCachedWidgetVisible;
+
 	/** The focusability of the widget at the point it was passed to Initialize */
 	bool bCachedWidgetFocusable;
 	
@@ -422,7 +430,10 @@ public:
 	 */
 	static FText GetWidgetVisibilityText(const TSharedPtr<SWidget>& InWidget);
 
-/**
+	/** Is the widget visible? */
+	static bool GetWidgetVisibility(const TSharedPtr<SWidget>& InWidget);
+
+	/**
 	 * @return The current clipping string for the given widget
 	 */
 	static FText GetWidgetClippingText(const TSharedPtr<SWidget>& InWidget);
diff --git a/Engine/Source/Developer/SlateReflector/Private/Widgets/SWidgetReflector.cpp b/Engine/Source/Developer/SlateReflector/Private/Widgets/SWidgetReflector.cpp
--- a/Engine/Source/Developer/SlateReflector/Private/Widgets/SWidgetReflector.cpp
+++ b/Engine/Source/Developer/SlateReflector/Private/Widgets/SWidgetReflector.cpp
@@ -240,6 +240,9 @@ private:
 	 */
 	int32 VisualizeSelectedNodesAsRectangles( const TArray<TSharedRef<FWidgetReflectorNodeBase>>& InNodesToDraw, const TSharedRef<SWindow>& VisualizeInWindow, FSlateWindowElementList& OutDrawElements, int32 LayerId );
 
+	/** Draw the actual highlight */
+	void DrawWidgetVisualization(const FPaintGeometry& WidgetGeometry, FLinearColor Color, FSlateWindowElementList& OutDrawElements, int32& LayerId);
+
 	/** Callback for changing the application scale slider. */
 	void HandleAppScaleSliderChanged( float NewValue )
 	{
@@ -1259,25 +1262,18 @@ int32 SWidgetReflector::VisualizePickAsRectangles( const FWidgetPath& InWidgetsT
 		FPaintGeometry WindowSpaceGeometry = WidgetGeometry.Geometry.ToPaintGeometry();
 		WindowSpaceGeometry.AppendTransform(TransformCast<FSlateLayoutTransform>(Inverse(InWidgetsToVisualize.TopLevelWindow->GetPositionInScreen())));
 
-		FSlateDrawElement::MakeBox(
-			OutDrawElements,
-			++LayerId,
-			WindowSpaceGeometry,
-			FCoreStyle::Get().GetBrush(TEXT("Debug.Border")),
-			ESlateDrawEffect::None,
-			FMath::Lerp(TopmostWidgetColor, LeafmostWidgetColor, ColorFactor)
-		);
+		FLinearColor Color = FMath::Lerp(TopmostWidgetColor, LeafmostWidgetColor, ColorFactor);
+		DrawWidgetVisualization(WindowSpaceGeometry, Color, OutDrawElements, LayerId);
 	}
 
 	return LayerId;
 }
 
-
 int32 SWidgetReflector::VisualizeSelectedNodesAsRectangles( const TArray<TSharedRef<FWidgetReflectorNodeBase>>& InNodesToDraw, const TSharedRef<SWindow>& VisualizeInWindow, FSlateWindowElementList& OutDrawElements, int32 LayerId )
 {
 	for (int32 NodeIndex = 0; NodeIndex < InNodesToDraw.Num(); ++NodeIndex)
 	{
-		const auto& NodeToDraw = InNodesToDraw[NodeIndex];
+		const TSharedRef<FWidgetReflectorNodeBase>& NodeToDraw = InNodesToDraw[NodeIndex];
 		const FLinearColor Tint(0.0f, 1.0f, 0.0f);
 
 		// The FGeometry we get is from a WidgetPath, so it's rooted in desktop space.
@@ -1288,17 +1284,50 @@ int32 SWidgetReflector::VisualizeSelectedNodesAsRectangles( const TArray<TShared
 		FPaintGeometry WindowSpaceGeometry(NodeToDraw->GetAccumulatedLayoutTransform(), NodeToDraw->GetAccumulatedRenderTransform(), NodeToDraw->GetLocalSize(), NodeToDraw->GetGeometry().HasRenderTransform());
 		WindowSpaceGeometry.AppendTransform(TransformCast<FSlateLayoutTransform>(Inverse(VisualizeInWindow->GetPositionInScreen())));
 
+		DrawWidgetVisualization(WindowSpaceGeometry, NodeToDraw->GetTint(), OutDrawElements, LayerId);
+	}
+
+	return LayerId;
+}
+
+void SWidgetReflector::DrawWidgetVisualization(const FPaintGeometry& WidgetGeometry, FLinearColor Color, FSlateWindowElementList& OutDrawElements, int32& LayerId)
+{
+	WidgetGeometry.CommitTransformsIfUsingLegacyConstructor();
+	const FVector2D LocalSize = WidgetGeometry.GetLocalSize();
+
+	// If the size is 0 in any dimension, we're going to draw a line to represent the widget, since it's going to take up
+	// padding space since it's visible, even though it's zero sized.
+	if (FMath::IsNearlyZero(LocalSize.X) || FMath::IsNearlyZero(LocalSize.Y))
+	{
+		TArray<FVector2D> LinePoints;
+		LinePoints.SetNum(2);
+
+		LinePoints[0] = FVector2D::ZeroVector;
+		LinePoints[1] = LocalSize;
+
+		FSlateDrawElement::MakeLines(
+			OutDrawElements,
+			++LayerId,
+			WidgetGeometry,
+			LinePoints,
+			ESlateDrawEffect::None,
+			Color,
+			true,
+			2
+		);
+	}
+	else
+	{
+		// Draw a normal box border around the geometry
 		FSlateDrawElement::MakeBox(
 			OutDrawElements,
 			++LayerId,
-			WindowSpaceGeometry,
+			WidgetGeometry,
 			FCoreStyle::Get().GetBrush(TEXT("Debug.Border")),
 			ESlateDrawEffect::None,
-			NodeToDraw->GetTint()
+			Color
 		);
 	}
-
-	return LayerId;
 }
 
 
diff --git a/Engine/Source/Developer/SlateReflector/Private/Widgets/SWidgetReflectorTreeWidgetItem.cpp b/Engine/Source/Developer/SlateReflector/Private/Widgets/SWidgetReflectorTreeWidgetItem.cpp
--- a/Engine/Source/Developer/SlateReflector/Private/Widgets/SWidgetReflectorTreeWidgetItem.cpp
+++ b/Engine/Source/Developer/SlateReflector/Private/Widgets/SWidgetReflectorTreeWidgetItem.cpp
@@ -33,6 +33,7 @@ void SReflectorTreeWidgetItem::Construct(const FArguments& InArgs, const TShared
 	CachedWidgetVisibility = WidgetInfo->GetWidgetVisibilityText();
 	CachedWidgetClipping = WidgetInfo->GetWidgetClippingText();
 	bCachedWidgetFocusable = WidgetInfo->GetWidgetFocusable();
+	bCachedWidgetVisible = WidgetInfo->GetWidgetVisible();
 	CachedReadableLocation = WidgetInfo->GetWidgetReadableLocation();
 	CachedWidgetFile = WidgetInfo->GetWidgetFile();
 	CachedWidgetLineNumber = WidgetInfo->GetWidgetLineNumber();
diff --git a/Engine/Source/Developer/SlateReflector/Private/Widgets/SWidgetReflectorTreeWidgetItem.h b/Engine/Source/Developer/SlateReflector/Private/Widgets/SWidgetReflectorTreeWidgetItem.h
--- a/Engine/Source/Developer/SlateReflector/Private/Widgets/SWidgetReflectorTreeWidgetItem.h
+++ b/Engine/Source/Developer/SlateReflector/Private/Widgets/SWidgetReflectorTreeWidgetItem.h
@@ -120,6 +120,7 @@ private:
 	FText CachedWidgetVisibility;
 	FText CachedWidgetClipping;
 	bool bCachedWidgetFocusable;
+	bool bCachedWidgetVisible;
 	FText CachedReadableLocation;
 	FString CachedWidgetFile;
 	int32 CachedWidgetLineNumber;
