SlateGI - Simages now cache their last assigned image, if they don't it's possible the memory changed out from under them, and they have no idea if they actually need to invalidate.  Can still further improve images by checking if the image size changed to perform a paint instead of a layout invalidate.  Adding a SetAttribute function on SWidget that makes setting values and invalidation an easier pattern to repeat effiecently and not need to overinvalidate things like volatility.  Updating several places in the widget reflector so that we can capture the tick/volatility/indirectvolatility/has active timers, state of widgets in the widget tree.

Matt.Kuhlenschmidt

diff --git a/Engine/Source/Developer/SlateReflector/Private/Models/WidgetReflectorNode.cpp b/Engine/Source/Developer/SlateReflector/Private/Models/WidgetReflectorNode.cpp
--- a/Engine/Source/Developer/SlateReflector/Private/Models/WidgetReflectorNode.cpp
+++ b/Engine/Source/Developer/SlateReflector/Private/Models/WidgetReflectorNode.cpp
@@ -136,6 +136,26 @@ bool FLiveWidgetReflectorNode::GetWidgetFocusable() const
 	return FWidgetReflectorNodeUtils::GetWidgetFocusable(Widget.Pin());
 }
 
+bool FLiveWidgetReflectorNode::GetWidgetNeedsTick() const
+{
+	return FWidgetReflectorNodeUtils::GetWidgetNeedsTick(Widget.Pin());
+}
+
+bool FLiveWidgetReflectorNode::GetWidgetIsVolatile() const
+{
+	return FWidgetReflectorNodeUtils::GetWidgetIsVolatile(Widget.Pin());
+}
+
+bool FLiveWidgetReflectorNode::GetWidgetIsVolatileIndirectly() const
+{
+	return FWidgetReflectorNodeUtils::GetWidgetIsVolatileIndirectly(Widget.Pin());
+}
+
+bool FLiveWidgetReflectorNode::GetWidgetHasActiveTimers() const
+{
+	return FWidgetReflectorNodeUtils::GetWidgetHasActiveTimers(Widget.Pin());
+}
+
 FText FLiveWidgetReflectorNode::GetWidgetReadableLocation() const
 {
 	return FWidgetReflectorNodeUtils::GetWidgetReadableLocation(Widget.Pin());
@@ -205,6 +225,10 @@ FSnapshotWidgetReflectorNode::FSnapshotWidgetReflectorNode(const FArrangedWidget
 	, CachedWidgetVisibilityText(FWidgetReflectorNodeUtils::GetWidgetVisibilityText(InArrangedWidget.Widget))
 	, bCachedWidgetVisible(FWidgetReflectorNodeUtils::GetWidgetVisibility(InArrangedWidget.Widget))
 	, bCachedWidgetFocusable(FWidgetReflectorNodeUtils::GetWidgetFocusable(InArrangedWidget.Widget))
+	, CachedWidgetNeedsTick(FWidgetReflectorNodeUtils::GetWidgetNeedsTick(InArrangedWidget.Widget))
+	, CachedWidgetIsVolatile(FWidgetReflectorNodeUtils::GetWidgetIsVolatile(InArrangedWidget.Widget))
+	, CachedWidgetIsVolatileIndirectly(FWidgetReflectorNodeUtils::GetWidgetIsVolatileIndirectly(InArrangedWidget.Widget))
+	, CachedWidgetHasActiveTimers(FWidgetReflectorNodeUtils::GetWidgetHasActiveTimers(InArrangedWidget.Widget))
 	, CachedWidgetClippingText(FWidgetReflectorNodeUtils::GetWidgetClippingText(InArrangedWidget.Widget))
 	, CachedWidgetReadableLocation(FWidgetReflectorNodeUtils::GetWidgetReadableLocation(InArrangedWidget.Widget))
 	, CachedWidgetFile(FWidgetReflectorNodeUtils::GetWidgetFile(InArrangedWidget.Widget))
@@ -257,6 +281,26 @@ FText FSnapshotWidgetReflectorNode::GetWidgetClippingText() const
 	return CachedWidgetClippingText;
 }
 
+bool FSnapshotWidgetReflectorNode::GetWidgetNeedsTick() const
+{
+	return CachedWidgetNeedsTick;
+}
+
+bool FSnapshotWidgetReflectorNode::GetWidgetIsVolatile() const
+{
+	return CachedWidgetIsVolatile;
+}
+
+bool FSnapshotWidgetReflectorNode::GetWidgetIsVolatileIndirectly() const
+{
+	return CachedWidgetIsVolatileIndirectly;
+}
+
+bool FSnapshotWidgetReflectorNode::GetWidgetHasActiveTimers() const
+{
+	return CachedWidgetHasActiveTimers;
+}
+
 FText FSnapshotWidgetReflectorNode::GetWidgetReadableLocation() const
 {
 	return CachedWidgetReadableLocation;
@@ -660,6 +704,25 @@ bool FWidgetReflectorNodeUtils::GetWidgetFocusable(const TSharedPtr<SWidget>& In
 	return InWidget.IsValid() ? InWidget->SupportsKeyboardFocus() : false;
 }
 
+bool FWidgetReflectorNodeUtils::GetWidgetNeedsTick(const TSharedPtr<SWidget>& InWidget)
+{
+	return InWidget.IsValid() ? InWidget->GetCanTick() : false;
+}
+
+bool FWidgetReflectorNodeUtils::GetWidgetIsVolatile(const TSharedPtr<SWidget>& InWidget)
+{
+	return InWidget.IsValid() ? InWidget->IsVolatile() : false;
+}
+
+bool FWidgetReflectorNodeUtils::GetWidgetIsVolatileIndirectly(const TSharedPtr<SWidget>& InWidget)
+{
+	return InWidget.IsValid() ? InWidget->IsVolatileIndirectly() : false;
+}
+
+bool FWidgetReflectorNodeUtils::GetWidgetHasActiveTimers(const TSharedPtr<SWidget>& InWidget)
+{
+	return InWidget.IsValid() ? InWidget->HasActiveTimers() : false;
+}
 
 FText FWidgetReflectorNodeUtils::GetWidgetClippingText(const TSharedPtr<SWidget>& InWidget)
 {
diff --git a/Engine/Source/Developer/SlateReflector/Private/Models/WidgetReflectorNode.h b/Engine/Source/Developer/SlateReflector/Private/Models/WidgetReflectorNode.h
--- a/Engine/Source/Developer/SlateReflector/Private/Models/WidgetReflectorNode.h
+++ b/Engine/Source/Developer/SlateReflector/Private/Models/WidgetReflectorNode.h
@@ -82,6 +82,11 @@ public:
 	/** Is the widget visible? */
 	virtual bool GetWidgetVisible() const = 0;
 
+	virtual bool GetWidgetNeedsTick() const = 0;
+	virtual bool GetWidgetIsVolatile() const = 0;
+	virtual bool GetWidgetIsVolatileIndirectly() const = 0;
+	virtual bool GetWidgetHasActiveTimers() const = 0;
+
 	/**
 	 * The human readable location for widgets that are defined in C++ is the file and line number
 	 * The human readable location for widgets that are defined in UMG is the asset name
@@ -224,6 +229,10 @@ public:
 	virtual bool GetWidgetVisible() const override;
 	virtual FText GetWidgetClippingText() const override;
 	virtual bool GetWidgetFocusable() const override;
+	virtual bool GetWidgetNeedsTick() const override;
+	virtual bool GetWidgetIsVolatile() const override;
+	virtual bool GetWidgetIsVolatileIndirectly() const override;
+	virtual bool GetWidgetHasActiveTimers() const override;
 	virtual FText GetWidgetReadableLocation() const override;
 	virtual FString GetWidgetFile() const override;
 	virtual int32 GetWidgetLineNumber() const override;
@@ -275,6 +284,10 @@ public:
 	virtual FText GetWidgetClippingText() const override;
 	virtual bool GetWidgetVisible() const override;
 	virtual bool GetWidgetFocusable() const override;
+	virtual bool GetWidgetNeedsTick() const override;
+	virtual bool GetWidgetIsVolatile() const override;
+	virtual bool GetWidgetIsVolatileIndirectly() const override;
+	virtual bool GetWidgetHasActiveTimers() const override;
 	virtual FText GetWidgetReadableLocation() const override;
 	virtual FString GetWidgetFile() const override;
 	virtual int32 GetWidgetLineNumber() const override;
@@ -317,8 +330,13 @@ private:
 
 	/** The focusability of the widget at the point it was passed to Initialize */
 	bool bCachedWidgetFocusable;
+
+	bool CachedWidgetNeedsTick;
+	bool CachedWidgetIsVolatile;
+	bool CachedWidgetIsVolatileIndirectly;
+	bool CachedWidgetHasActiveTimers;
 	
-		/** The clipping string of the widget at the point it was passed to Initialize */
+	/** The clipping string of the widget at the point it was passed to Initialize */
 	FText CachedWidgetClippingText;
 
 	/** The human readable location (source file for C++ widgets, asset name for UMG widgets) of the widget at the point it was passed to Initialize */
@@ -443,6 +461,10 @@ public:
 	*/
 	static bool GetWidgetFocusable(const TSharedPtr<SWidget>& InWidget);
 
+	static bool GetWidgetNeedsTick(const TSharedPtr<SWidget>& InWidget);
+	static bool GetWidgetIsVolatile(const TSharedPtr<SWidget>& InWidget);
+	static bool GetWidgetIsVolatileIndirectly(const TSharedPtr<SWidget>& InWidget);
+	static bool GetWidgetHasActiveTimers(const TSharedPtr<SWidget>& InWidget);
 	
 	/**
 	 * The human readable location for widgets that are defined in C++ is the file and line number
diff --git a/Engine/Source/Developer/SlateReflector/Private/Widgets/SWidgetReflectorToolTipWidget.h b/Engine/Source/Developer/SlateReflector/Private/Widgets/SWidgetReflectorToolTipWidget.h
--- a/Engine/Source/Developer/SlateReflector/Private/Widgets/SWidgetReflectorToolTipWidget.h
+++ b/Engine/Source/Developer/SlateReflector/Private/Widgets/SWidgetReflectorToolTipWidget.h
@@ -73,33 +73,93 @@ public:
 
 			// Size Info
 			+ SGridPanel::Slot(0, 2)
-				[
-					SNew(STextBlock)
-						.Text(LOCTEXT("SizeInfo", "Size Info"))
-				]
+			[
+				SNew(STextBlock)
+				.Text(LOCTEXT("SizeInfo", "Size Info"))
+			]
 
 			// Size Info Value
 			+ SGridPanel::Slot(1, 2)
-				.Padding(5.0f, 0.0f, 0.0f, 0.0f)
-				[
-					SNew(STextBlock)
-						.Text(this, &SReflectorToolTipWidget::GetSizeInfo)
-				]
+			.Padding(5.0f, 0.0f, 0.0f, 0.0f)
+			[
+				SNew(STextBlock)
+				.Text(this, &SReflectorToolTipWidget::GetSizeInfo)
+			]
 
 			// Enabled
 			+ SGridPanel::Slot(0, 3)
-				[
-					SNew(STextBlock)
-						.Text(LOCTEXT("Enabled", "Enabled"))
-				]
+			[
+				SNew(STextBlock)
+				.Text(LOCTEXT("Enabled", "Enabled"))
+			]
 
 			// Enabled Value
 			+ SGridPanel::Slot(1, 3)
-				.Padding(5.0f, 0.0f, 0.0f, 0.0f)
-				[
-					SNew(STextBlock)
-						.Text(this, &SReflectorToolTipWidget::GetEnabled)
-				]
+			.Padding(5.0f, 0.0f, 0.0f, 0.0f)
+			[
+				SNew(STextBlock)
+				.Text(this, &SReflectorToolTipWidget::GetEnabled)
+			]
+
+			// NeedsTick
+			+ SGridPanel::Slot(0, 4)
+			[
+				SNew(STextBlock)
+				.Text(LOCTEXT("NeedsTick", "Needs Tick"))
+			]
+
+			// NeedsTick Value
+			+ SGridPanel::Slot(1, 4)
+			.Padding(5.0f, 0.0f, 0.0f, 0.0f)
+			[
+				SNew(STextBlock)
+				.Text(this, &SReflectorToolTipWidget::GetNeedsTick)
+			]
+
+			// IsVolatile
+			+ SGridPanel::Slot(0, 5)
+			[
+				SNew(STextBlock)
+				.Text(LOCTEXT("IsVolatile", "Is Volatile"))
+			]
+
+			// IsVolatile Value
+			+ SGridPanel::Slot(1, 5)
+			.Padding(5.0f, 0.0f, 0.0f, 0.0f)
+			[
+				SNew(STextBlock)
+				.Text(this, &SReflectorToolTipWidget::GetIsVolatile)
+			]
+
+			// IsVolatileIndirectly
+			+ SGridPanel::Slot(0, 6)
+			[
+				SNew(STextBlock)
+				.Text(LOCTEXT("IsVolatileIndirectly", "Is Volatile Indirectly"))
+			]
+
+			// IsVolatileIndirectly Value
+			+ SGridPanel::Slot(1, 6)
+			.Padding(5.0f, 0.0f, 0.0f, 0.0f)
+			[
+				SNew(STextBlock)
+				.Text(this, &SReflectorToolTipWidget::GetIsVolatileIndirectly)
+			]
+
+			// HasActiveTimers
+			+ SGridPanel::Slot(0, 7)
+			[
+				SNew(STextBlock)
+				.Text(LOCTEXT("HasActiveTimers", "Has Active Timers"))
+			]
+
+			// HasActiveTimers Value
+			+ SGridPanel::Slot(1, 7)
+			.Padding(5.0f, 0.0f, 0.0f, 0.0f)
+			[
+				SNew(STextBlock)
+				.Text(this, &SReflectorToolTipWidget::GetHasActiveTimers)
+			]
 		];
 	}
 
@@ -127,6 +187,34 @@ private:
 		return (WidgetInfo->GetWidgetEnabled()) ? TrueText : FalseText;
 	}
 
+	FText GetNeedsTick() const
+	{
+		static const FText TrueText = LOCTEXT("True", "True");
+		static const FText FalseText = LOCTEXT("False", "False");
+		return (WidgetInfo->GetWidgetNeedsTick()) ? TrueText : FalseText;
+	}
+
+	FText GetIsVolatile() const
+	{
+		static const FText TrueText = LOCTEXT("True", "True");
+		static const FText FalseText = LOCTEXT("False", "False");
+		return (WidgetInfo->GetWidgetIsVolatile()) ? TrueText : FalseText;
+	}
+
+	FText GetIsVolatileIndirectly() const
+	{
+		static const FText TrueText = LOCTEXT("True", "True");
+		static const FText FalseText = LOCTEXT("False", "False");
+		return (WidgetInfo->GetWidgetIsVolatileIndirectly()) ? TrueText : FalseText;
+	}
+
+	FText GetHasActiveTimers() const
+	{
+		static const FText TrueText = LOCTEXT("True", "True");
+		static const FText FalseText = LOCTEXT("False", "False");
+		return (WidgetInfo->GetWidgetHasActiveTimers()) ? TrueText : FalseText;
+	}
+
 private:
 
 	/** The info about the widget that we are visualizing. */
diff --git a/Engine/Source/Developer/SlateReflector/Private/Widgets/SWidgetReflectorTreeWidgetItem.cpp b/Engine/Source/Developer/SlateReflector/Private/Widgets/SWidgetReflectorTreeWidgetItem.cpp
--- a/Engine/Source/Developer/SlateReflector/Private/Widgets/SWidgetReflectorTreeWidgetItem.cpp
+++ b/Engine/Source/Developer/SlateReflector/Private/Widgets/SWidgetReflectorTreeWidgetItem.cpp
@@ -34,6 +34,10 @@ void SReflectorTreeWidgetItem::Construct(const FArguments& InArgs, const TShared
 	CachedWidgetClipping = WidgetInfo->GetWidgetClippingText();
 	bCachedWidgetFocusable = WidgetInfo->GetWidgetFocusable();
 	bCachedWidgetVisible = WidgetInfo->GetWidgetVisible();
+	bCachedWidgetNeedsTick = WidgetInfo->GetWidgetNeedsTick();
+	bCachedWidgetIsVolatile = WidgetInfo->GetWidgetIsVolatile();
+	bCachedWidgetIsVolatileIndirectly = WidgetInfo->GetWidgetIsVolatileIndirectly();
+	bCachedWidgetHasActiveTimers = WidgetInfo->GetWidgetHasActiveTimers();
 	CachedReadableLocation = WidgetInfo->GetWidgetReadableLocation();
 	CachedWidgetFile = WidgetInfo->GetWidgetFile();
 	CachedWidgetLineNumber = WidgetInfo->GetWidgetLineNumber();
diff --git a/Engine/Source/Developer/SlateReflector/Private/Widgets/SWidgetReflectorTreeWidgetItem.h b/Engine/Source/Developer/SlateReflector/Private/Widgets/SWidgetReflectorTreeWidgetItem.h
--- a/Engine/Source/Developer/SlateReflector/Private/Widgets/SWidgetReflectorTreeWidgetItem.h
+++ b/Engine/Source/Developer/SlateReflector/Private/Widgets/SWidgetReflectorTreeWidgetItem.h
@@ -121,6 +121,10 @@ private:
 	FText CachedWidgetClipping;
 	bool bCachedWidgetFocusable;
 	bool bCachedWidgetVisible;
+	bool bCachedWidgetNeedsTick;
+	bool bCachedWidgetIsVolatile;
+	bool bCachedWidgetIsVolatileIndirectly;
+	bool bCachedWidgetHasActiveTimers;
 	FText CachedReadableLocation;
 	FString CachedWidgetFile;
 	int32 CachedWidgetLineNumber;
diff --git a/Engine/Source/Runtime/SlateCore/Private/Widgets/Images/SImage.cpp b/Engine/Source/Runtime/SlateCore/Private/Widgets/Images/SImage.cpp
--- a/Engine/Source/Runtime/SlateCore/Private/Widgets/Images/SImage.cpp
+++ b/Engine/Source/Runtime/SlateCore/Private/Widgets/Images/SImage.cpp
@@ -13,8 +13,6 @@ void SImage::Construct( const FArguments& InArgs )
 	ColorAndOpacity = InArgs._ColorAndOpacity;
 	bFlipForRightToLeftFlowDirection = InArgs._FlipForRightToLeftFlowDirection;
 	SetOnMouseButtonDown(InArgs._OnMouseButtonDown);
-
-
 }
 
 int32 SImage::OnPaint( const FPaintArgs& Args, const FGeometry& AllottedGeometry, const FSlateRect& MyCullingRect, FSlateWindowElementList& OutDrawElements, int32 LayerId, const FWidgetStyle& InWidgetStyle, bool bParentEnabled ) const
@@ -54,29 +52,32 @@ FVector2D SImage::ComputeDesiredSize( float ) const
 
 void SImage::SetColorAndOpacity( const TAttribute<FSlateColor>& InColorAndOpacity )
 {
-	if (!ColorAndOpacity.IdenticalTo(InColorAndOpacity))
-	{
-		ColorAndOpacity = InColorAndOpacity;
-		Invalidate(EInvalidateWidget::PaintAndVolatility);
-	}
+	SetAttribute(ColorAndOpacity, InColorAndOpacity, EInvalidateWidgetReason::Paint);
 }
 
 void SImage::SetColorAndOpacity( FLinearColor InColorAndOpacity )
 {
-	if (!ColorAndOpacity.IdenticalTo(InColorAndOpacity))
-	{
-		ColorAndOpacity = InColorAndOpacity;
-		Invalidate(EInvalidateWidget::PaintAndVolatility);
-	}
+	SetColorAndOpacity(TAttribute<FSlateColor>(InColorAndOpacity));
 }
 
 void SImage::SetImage(TAttribute<const FSlateBrush*> InImage)
 {
-	if (!Image.IdenticalTo(InImage))
+	const bool bImagePointerChanged = SetAttribute(Image, InImage, EInvalidateWidgetReason::Layout);
+
+	const FSlateBrush* ImagePtr = Image.Get();
+	const FSlateBrush NewImageToCache = ImagePtr ? *ImagePtr : FSlateBrush();
+
+	// If the slate brush pointer didn't change, that may not mean nothing changed.  We
+	// sometimes reuse the slate brush memory address and change out the texture.  In those
+	// circumstances we need to actually look at the data the brush has compared to what it had
+	// previously.
+	if (!bImagePointerChanged && ImageCache != NewImageToCache)
 	{
-		Image = InImage;
-		Invalidate(EInvalidateWidget::LayoutAndVolatility);
+		Invalidate(EInvalidateWidgetReason::Layout);
 	}
+
+	// Cache the new image's value in case the memory changes later.
+	ImageCache = NewImageToCache;
 }
 
 #if WITH_ACCESSIBILITY
diff --git a/Engine/Source/Runtime/SlateCore/Private/Widgets/SWidget.cpp b/Engine/Source/Runtime/SlateCore/Private/Widgets/SWidget.cpp
--- a/Engine/Source/Runtime/SlateCore/Private/Widgets/SWidget.cpp
+++ b/Engine/Source/Runtime/SlateCore/Private/Widgets/SWidget.cpp
@@ -1023,12 +1023,7 @@ bool SWidget::IsDirectlyHovered() const
 
 void SWidget::SetVisibility(TAttribute<EVisibility> InVisibility)
 {
-	if (!Visibility.IdenticalTo(InVisibility))
-	{
-		Visibility = InVisibility;
-
-		Invalidate(EInvalidateWidget::Visibility);
-	}
+	SetAttribute(Visibility, InVisibility, EInvalidateWidgetReason::Visibility);
 }
 
 void SWidget::Invalidate(EInvalidateWidget InvalidateReason)
diff --git a/Engine/Source/Runtime/SlateCore/Public/Widgets/Images/SImage.h b/Engine/Source/Runtime/SlateCore/Public/Widgets/Images/SImage.h
--- a/Engine/Source/Runtime/SlateCore/Public/Widgets/Images/SImage.h
+++ b/Engine/Source/Runtime/SlateCore/Public/Widgets/Images/SImage.h
@@ -81,6 +81,9 @@ protected:
 
 	/** The slate brush to draw for the image, or a bound delegate to a brush. */
 	TAttribute< const FSlateBrush* > Image;
+	
+	/** The copy of the image data, some users reuse the same FSlateBrush pointer, so we need to check it against the last true data to see what changed. */
+	FSlateBrush ImageCache;
 
 	/** Color and opacity scale for this image */
 	TAttribute<FSlateColor> ColorAndOpacity;
diff --git a/Engine/Source/Runtime/SlateCore/Public/Widgets/SWidget.h b/Engine/Source/Runtime/SlateCore/Public/Widgets/SWidget.h
--- a/Engine/Source/Runtime/SlateCore/Public/Widgets/SWidget.h
+++ b/Engine/Source/Runtime/SlateCore/Public/Widgets/SWidget.h
@@ -924,11 +924,7 @@ public:
 	 */
 	void SetEnabled(const TAttribute<bool>& InEnabledState)
 	{
-		if (!EnabledState.IdenticalTo(InEnabledState))
-		{
-			EnabledState = InEnabledState;
-			Invalidate(EInvalidateWidget::PaintAndVolatility);
-		}
+		SetAttribute(EnabledState, InEnabledState, EInvalidateWidgetReason::Paint);
 	}
 
 	/** @return Whether or not this widget is enabled */
@@ -1162,11 +1158,7 @@ public:
 	/** @param InTransform the render transform to set for the widget (transforms from widget's local space). TOptional<> to allow code to skip expensive overhead if there is no render transform applied. */
 	FORCEINLINE void SetRenderTransform(TAttribute<TOptional<FSlateRenderTransform>> InTransform)
 	{
-		if(!RenderTransform.IdenticalTo(InTransform))
-		{
-			RenderTransform = InTransform;
-			Invalidate(EInvalidateWidget::LayoutAndVolatility|EInvalidateWidget::RenderTransform);
-		}
+		SetAttribute(RenderTransform, InTransform, EInvalidateWidgetReason::Layout | EInvalidateWidgetReason::RenderTransform);
 	}
 
 	/** @return the pivot point of the render transform. */
@@ -1178,11 +1170,7 @@ public:
 	/** @param InTransformPivot Sets the pivot point of the widget's render transform (in normalized local space). */
 	FORCEINLINE void SetRenderTransformPivot(TAttribute<FVector2D> InTransformPivot)
 	{
-		if (!RenderTransformPivot.IdenticalTo(InTransformPivot))
-		{
-			RenderTransformPivot = InTransformPivot;
-			Invalidate(EInvalidateWidget::LayoutAndVolatility | EInvalidateWidget::RenderTransform);
-		}
+		SetAttribute(RenderTransformPivot, InTransformPivot, EInvalidateWidgetReason::Layout | EInvalidateWidgetReason::RenderTransform);
 	}
 
 	/**
@@ -1529,6 +1517,9 @@ public:
 	 * Unregisters an active timer handle. This is optional, as the delegate can UnRegister itself by returning EActiveTimerReturnType::Stop.
 	 */
 	void UnRegisterActiveTimer( const TSharedRef<FActiveTimerHandle>& ActiveTimerHandle );
+	
+	/** Does this widget have any active timers? */
+	bool HasActiveTimers() const { return ActiveTimers.Num() > 0; }
 
 private:
 
@@ -1537,6 +1528,34 @@ private:
 
 	const FPointerEventHandler* GetPointerEvent(const FName EventName) const;
 	void SetPointerEvent(const FName EventName, FPointerEventHandler& InEvent);
+
+protected:
+	/**
+	 * Performs the attribute assignment and invalidates the widget minimally based on what actually changed.  So if the boundness of the attribute didn't change
+	 * volatility won't need to be recalculated.  Returns true if the value changed.
+	 */
+	template<typename TargetValueType, typename SourceValueType>
+	bool SetAttribute(TAttribute<TargetValueType>& TargetValue, const TAttribute<SourceValueType>& SourceValue, EInvalidateWidgetReason BaseInvalidationReason)
+	{
+		if (!TargetValue.IdenticalTo(SourceValue))
+		{
+			const bool bWasBound = TargetValue.IsBound();
+			const bool bBoundnessChanged = bWasBound != SourceValue.IsBound();
+			TargetValue = SourceValue;
+
+			EInvalidateWidgetReason InvalidateReason = BaseInvalidationReason;
+			if (bBoundnessChanged)
+			{
+				InvalidateReason |= EInvalidateWidgetReason::Volatility;
+			}
+
+			Invalidate(InvalidateReason);
+			return true;
+		}
+
+		return false;
+	}
+
 protected:
 	/** Dtor ensures that active timer handles are UnRegistered with the SlateApplication. */
 	virtual ~SWidget();
diff --git a/Engine/Source/Runtime/UMG/Private/Components/Image.cpp b/Engine/Source/Runtime/UMG/Private/Components/Image.cpp
--- a/Engine/Source/Runtime/UMG/Private/Components/Image.cpp
+++ b/Engine/Source/Runtime/UMG/Private/Components/Image.cpp
@@ -100,7 +100,6 @@ void UImage::SetBrush(const FSlateBrush& InBrush)
 		if (MyImage.IsValid())
 		{
 			MyImage->SetImage(&Brush);
-			MyImage->Invalidate(EInvalidateWidget::LayoutAndVolatility);
 		}
 	}
 }
@@ -114,7 +113,6 @@ void UImage::SetBrushSize(FVector2D DesiredSize)
 		if (MyImage.IsValid())
 		{
 			MyImage->SetImage(&Brush);
-			MyImage->Invalidate(EInvalidateWidget::LayoutAndVolatility);
 		}
 	}
 }
@@ -128,7 +126,6 @@ void UImage::SetBrushTintColor(FSlateColor TintColor)
 		if (MyImage.IsValid())
 		{
 			MyImage->SetImage(&Brush);
-			MyImage->Invalidate(EInvalidateWidget::PaintAndVolatility);
 		}
 	}
 }
@@ -142,7 +139,6 @@ void UImage::SetBrushResourceObject(UObject* ResourceObject)
 		if (MyImage.IsValid())
 		{
 			MyImage->SetImage(&Brush);
-			MyImage->Invalidate(EInvalidateWidget::PaintAndVolatility);
 		}
 	}
 }
@@ -157,7 +153,6 @@ void UImage::SetBrushFromAsset(USlateBrushAsset* Asset)
 		if (MyImage.IsValid())
 		{
 			MyImage->SetImage(&Brush);
-			MyImage->Invalidate(EInvalidateWidget::LayoutAndVolatility);
 		}
 	}
 }
@@ -191,7 +186,6 @@ void UImage::SetBrushFromTexture(UTexture2D* Texture, bool bMatchSize)
 		if (MyImage.IsValid())
 		{
 			MyImage->SetImage(&Brush);
-			MyImage->Invalidate(EInvalidateWidget::LayoutAndVolatility);
 		}
 	}
 }
@@ -219,7 +213,6 @@ void UImage::SetBrushFromAtlasInterface(TScriptInterface<ISlateTextureAtlasInter
 		if (MyImage.IsValid())
 		{
 			MyImage->SetImage(&Brush);
-			MyImage->Invalidate(EInvalidateWidget::LayoutAndVolatility);
 		}
 	}
 }
@@ -240,7 +233,6 @@ void UImage::SetBrushFromTextureDynamic(UTexture2DDynamic* Texture, bool bMatchS
 		if (MyImage.IsValid())
 		{
 			MyImage->SetImage(&Brush);
-			MyImage->Invalidate(EInvalidateWidget::LayoutAndVolatility);
 		}
 	}
 }
@@ -257,7 +249,6 @@ void UImage::SetBrushFromMaterial(UMaterialInterface* Material)
 		if (MyImage.IsValid())
 		{
 			MyImage->SetImage(&Brush);
-			MyImage->Invalidate(EInvalidateWidget::LayoutAndVolatility);
 		}
 	}
 }
@@ -373,7 +364,6 @@ UMaterialInstanceDynamic* UImage::GetDynamicMaterial()
 			if ( MyImage.IsValid() )
 			{
 				MyImage->SetImage(&Brush);
-				MyImage->Invalidate(EInvalidateWidget::LayoutAndVolatility);
 			}
 		}
 		
