Slate - Avoid a rare crash in the Widget Reflector when the clipping states don't match.

diff --git a/Engine/Source/Developer/SlateReflector/Private/VisualTreeCapture.cpp b/Engine/Source/Developer/SlateReflector/Private/VisualTreeCapture.cpp
--- a/Engine/Source/Developer/SlateReflector/Private/VisualTreeCapture.cpp
+++ b/Engine/Source/Developer/SlateReflector/Private/VisualTreeCapture.cpp
@@ -61,9 +61,16 @@ TSharedPtr<const SWidget> FVisualTreeSnapshot::Pick(FVector2D Point)
 	for (int Index = Entries.Num() - 1; Index >= 0; Index--)
 	{
 		const FVisualEntry& Entry = Entries[Index];
-		if (Entry.ClippingIndex != -1 && !ClippingStates[Entry.ClippingIndex].IsPointInside(Point))
+		if (Entry.ClippingIndex != -1)
 		{
-			continue;
+			// TODO: This isn't always a valid index, but should be.
+			if (ensure(ClippingStates.IsValidIndex(Entry.ClippingIndex)))
+			{
+				if (!ClippingStates[Entry.ClippingIndex].IsPointInside(Point))
+				{
+					continue;
+				}
+			}
 		}
 
 		if (!Entry.IsPointInside(Point))
