Cache off the value of slate global invalidation mode when using Pick against painted widgets, and then turns it off while picking painted widgets.  Eventually this should be modified to work against proxies and cached elements instead.

diff --git a/Engine/Source/Developer/SlateReflector/Private/Widgets/SWidgetReflector.cpp b/Engine/Source/Developer/SlateReflector/Private/Widgets/SWidgetReflector.cpp
--- a/Engine/Source/Developer/SlateReflector/Private/Widgets/SWidgetReflector.cpp
+++ b/Engine/Source/Developer/SlateReflector/Private/Widgets/SWidgetReflector.cpp
@@ -282,11 +282,22 @@ private:
 
 	void SetPickingMode(EWidgetPickingMode InMode)
 	{
+#if WITH_SLATE_DEBUGGING
+		static auto CVarSlateGlobalInvalidation = IConsoleManager::Get().FindConsoleVariable(TEXT("Slate.EnableGlobalInvalidation"));
+#endif
+
 		if (PickingMode != InMode)
 		{
 			// Disable visual picking, and renable widget caching.
 #if WITH_SLATE_DEBUGGING
 			SInvalidationPanel::EnableInvalidationPanels(true);
+
+			if (PickingMode == EWidgetPickingMode::None)
+			{
+				bLastGlobalInvalidationState = CVarSlateGlobalInvalidation->GetBool();
+			}
+
+			CVarSlateGlobalInvalidation->Set(bLastGlobalInvalidationState);
 #endif
 			VisualCapture.Disable();
 
@@ -308,6 +319,7 @@ private:
 				VisualCapture.Enable();
 #if WITH_SLATE_DEBUGGING
 				SInvalidationPanel::EnableInvalidationPanels(false);
+				CVarSlateGlobalInvalidation->Set(false);
 #endif
 			}
 		}
@@ -422,6 +434,8 @@ private:
 
 	FVisualTreeCapture VisualCapture;
 
+	bool bLastGlobalInvalidationState = false;
+
 private:
 	// DEMO MODE
 	bool bEnableDemoMode;
