WidgetReflector: Show more detaill in the tree view (enabled, volatile, HasActiveTimer, ActualSize).

diff --git a/Engine/Source/Developer/SlateReflector/Private/Widgets/SWidgetReflector.cpp b/Engine/Source/Developer/SlateReflector/Private/Widgets/SWidgetReflector.cpp
--- a/Engine/Source/Developer/SlateReflector/Private/Widgets/SWidgetReflector.cpp
+++ b/Engine/Source/Developer/SlateReflector/Private/Widgets/SWidgetReflector.cpp
@@ -445,11 +445,17 @@ private:
 
 void SWidgetReflector::Construct( const FArguments& InArgs )
 {
+	// If saved, LoadSettings will override these varibles.
+	LastPickingMode = EWidgetPickingMode::HitTesting;
+	HiddenReflectorTreeColumns.Add(SReflectorTreeWidgetItem::NAME_Enabled.ToString());
+	HiddenReflectorTreeColumns.Add(SReflectorTreeWidgetItem::NAME_Volatile.ToString());
+	HiddenReflectorTreeColumns.Add(SReflectorTreeWidgetItem::NAME_HasActiveTimer.ToString());
+	HiddenReflectorTreeColumns.Add(SReflectorTreeWidgetItem::NAME_ActualSize.ToString());
+
 	LoadSettings();
 
 	CurrentUIMode = EWidgetReflectorUIMode::Live;
 	PickingMode = EWidgetPickingMode::None;
-	//LastPickingMode = EWidgetPickingMode::HitTesting; //initialized in LoadSettings
 	bFilterReflectorTreeRootWithUMG = false;
 
 	SnapshotDelay = 0.0f;
@@ -846,24 +852,38 @@ TSharedRef<SDockTab> SWidgetReflector::SpawnWidgetHierarchyTab(const FSpawnTabAr
 						+SHeaderRow::Column(SReflectorTreeWidgetItem::NAME_Visibility)
 						.DefaultLabel(LOCTEXT("Visibility", "Visibility"))
 						.DefaultTooltip(LOCTEXT("VisibilityTooltip", "Visibility"))
-						.FixedWidth(125.0f)
+						.ManualWidth(125.0f)
+						
+						+ SHeaderRow::Column(SReflectorTreeWidgetItem::NAME_Enabled)
+						.DefaultLabel(LOCTEXT("Enabled", "Enabled"))
+						.DefaultTooltip(LOCTEXT("EnabledToolTip", "Enabled"))
+						.ManualWidth(60.0f)
 
 						+ SHeaderRow::Column(SReflectorTreeWidgetItem::NAME_Focusable)
-						.DefaultLabel(LOCTEXT("Focus", "Focus?"))
+						.DefaultLabel(LOCTEXT("Focus", "Focus"))
 						.DefaultTooltip(LOCTEXT("FocusableTooltip", "Focusability (Note that for hit-test directional navigation to work it must be Focusable and \"Visible\"!)"))
-						.FixedWidth(50.0f)
+						.ManualWidth(60.0f)
+
+						+ SHeaderRow::Column(SReflectorTreeWidgetItem::NAME_HasActiveTimer)
+						.DefaultLabel(LOCTEXT("HasActiveTimer", "Timer"))
+						.DefaultTooltip(LOCTEXT("HasActiveTimerTooltip", "Has Active Timer"))
+						.ManualWidth(60.0f)
 
 						+SHeaderRow::Column(SReflectorTreeWidgetItem::NAME_Clipping)
 						.DefaultLabel(LOCTEXT("Clipping", "Clipping" ))
-						.FixedWidth(100.0f)
+						.ManualWidth(100.0f)
+
+						+ SHeaderRow::Column(SReflectorTreeWidgetItem::NAME_ActualSize)
+						.DefaultLabel(LOCTEXT("ActualSize", "Size"))
+						.ManualWidth(100.0f)
 
 						+SHeaderRow::Column(SReflectorTreeWidgetItem::NAME_WidgetInfo)
 						.DefaultLabel(LOCTEXT("Source", "Source" ))
-						.FillWidth(0.20f)
+						.ManualWidth(200.f)
 
 						+SHeaderRow::Column(SReflectorTreeWidgetItem::NAME_Address)
 						.DefaultLabel( LOCTEXT("Address", "Address") )
-						.FixedWidth(170.0f)
+						.ManualWidth(170.0f)
 					)
 				]
 			]
@@ -993,15 +1013,18 @@ void SWidgetReflector::SaveSettings()
 
 void SWidgetReflector::LoadSettings()
 {
-	int32 LastPickingModeAsInt = static_cast<int32>(EWidgetPickingMode::HitTesting);
-	GConfig->GetInt(TEXT("WidgetReflector"), TEXT("LastPickingMode"), LastPickingModeAsInt, *GEditorPerProjectIni);
-	LastPickingMode = ConvertToWidgetPickingMode(LastPickingModeAsInt);
-	if (LastPickingMode == EWidgetPickingMode::None)
+	if (GConfig->DoesSectionExist(TEXT("WidgetReflector"), *GEditorPerProjectIni))
 	{
-		LastPickingMode = EWidgetPickingMode::HitTesting;
-	}
+		int32 LastPickingModeAsInt = static_cast<int32>(EWidgetPickingMode::HitTesting);
+		GConfig->GetInt(TEXT("WidgetReflector"), TEXT("LastPickingMode"), LastPickingModeAsInt, *GEditorPerProjectIni);
+		LastPickingMode = ConvertToWidgetPickingMode(LastPickingModeAsInt);
+		if (LastPickingMode == EWidgetPickingMode::None)
+		{
+			LastPickingMode = EWidgetPickingMode::HitTesting;
+		}
 
-	GConfig->GetArray(TEXT("WidgetReflector"), TEXT("HiddenReflectorTreeColumns"), HiddenReflectorTreeColumns, *GEditorPerProjectIni);
+		GConfig->GetArray(TEXT("WidgetReflector"), TEXT("HiddenReflectorTreeColumns"), HiddenReflectorTreeColumns, *GEditorPerProjectIni);
+	}
 }
 
 
diff --git a/Engine/Source/Developer/SlateReflector/Private/Widgets/SWidgetReflectorTreeWidgetItem.cpp b/Engine/Source/Developer/SlateReflector/Private/Widgets/SWidgetReflectorTreeWidgetItem.cpp
--- a/Engine/Source/Developer/SlateReflector/Private/Widgets/SWidgetReflectorTreeWidgetItem.cpp
+++ b/Engine/Source/Developer/SlateReflector/Private/Widgets/SWidgetReflectorTreeWidgetItem.cpp
@@ -16,8 +16,12 @@ FName SReflectorTreeWidgetItem::NAME_WidgetName(TEXT("WidgetName"));
 FName SReflectorTreeWidgetItem::NAME_WidgetInfo(TEXT("WidgetInfo"));
 FName SReflectorTreeWidgetItem::NAME_Visibility(TEXT("Visibility"));
 FName SReflectorTreeWidgetItem::NAME_Focusable(TEXT("Focusable"));
+FName SReflectorTreeWidgetItem::NAME_Enabled(TEXT("Enabled"));
+FName SReflectorTreeWidgetItem::NAME_Volatile(TEXT("Volatile"));
+FName SReflectorTreeWidgetItem::NAME_HasActiveTimer(TEXT("HasActiveTimer"));
 FName SReflectorTreeWidgetItem::NAME_Clipping(TEXT("Clipping"));
 FName SReflectorTreeWidgetItem::NAME_ForegroundColor(TEXT("ForegroundColor"));
+FName SReflectorTreeWidgetItem::NAME_ActualSize(TEXT("ActualSize"));
 FName SReflectorTreeWidgetItem::NAME_Address(TEXT("Address"));
 
 void SReflectorTreeWidgetItem::Construct(const FArguments& InArgs, const TSharedRef<STableViewBase>& InOwnerTableView)
@@ -28,20 +32,6 @@ void SReflectorTreeWidgetItem::Construct(const FArguments& InArgs, const TShared
 	this->SetPadding(0);
 
 	check(WidgetInfo.IsValid());
-	CachedWidgetType = WidgetInfo->GetWidgetType();
-	CachedWidgetTypeAndShortName = WidgetInfo->GetWidgetTypeAndShortName();
-	CachedWidgetVisibility = WidgetInfo->GetWidgetVisibilityText();
-	CachedWidgetClipping = WidgetInfo->GetWidgetClippingText();
-	bCachedWidgetFocusable = WidgetInfo->GetWidgetFocusable();
-	bCachedWidgetVisible = WidgetInfo->GetWidgetVisible();
-	bCachedWidgetNeedsTick = WidgetInfo->GetWidgetNeedsTick();
-	bCachedWidgetIsVolatile = WidgetInfo->GetWidgetIsVolatile();
-	bCachedWidgetIsVolatileIndirectly = WidgetInfo->GetWidgetIsVolatileIndirectly();
-	bCachedWidgetHasActiveTimers = WidgetInfo->GetWidgetHasActiveTimers();
-	CachedReadableLocation = WidgetInfo->GetWidgetReadableLocation();
-	CachedWidgetFile = WidgetInfo->GetWidgetFile();
-	CachedWidgetLineNumber = WidgetInfo->GetWidgetLineNumber();
-	CachedAssetData = WidgetInfo->GetWidgetAssetData();
 
 	SMultiColumnTableRow< TSharedRef<FWidgetReflectorNodeBase> >::Construct(SMultiColumnTableRow< TSharedRef<FWidgetReflectorNodeBase> >::FArguments().Padding(0), InOwnerTableView);
 }
@@ -49,6 +39,20 @@ void SReflectorTreeWidgetItem::Construct(const FArguments& InArgs, const TShared
 BEGIN_SLATE_FUNCTION_BUILD_OPTIMIZATION
 TSharedRef<SWidget> SReflectorTreeWidgetItem::GenerateWidgetForColumn(const FName& ColumnName)
 {
+	SReflectorTreeWidgetItem* Self = this;
+	auto BuildCheckBox = [Self](bool bIsChecked)
+		{
+			return SNew(SBox)
+				.HAlign(HAlign_Center)
+				.VAlign(VAlign_Center)
+				.Padding(FMargin(2.0f, 0.0f))
+				[
+					SNew(SCheckBox)
+					.Style(FCoreStyle::Get(), TEXT("WidgetReflector.FocusableCheck"))
+					.IsChecked(bIsChecked ? ECheckBoxState::Checked : ECheckBoxState::Unchecked)
+				];
+		};
+
 	if (ColumnName == NAME_WidgetName )
 	{
 		return SNew(SHorizontalBox)
@@ -67,7 +71,7 @@ TSharedRef<SWidget> SReflectorTreeWidgetItem::GenerateWidgetForColumn(const FNam
 		.VAlign(VAlign_Center)
 		[
 			SNew(STextBlock)
-			.Text(this, &SReflectorTreeWidgetItem::GetWidgetTypeAndShortName)
+			.Text(WidgetInfo->GetWidgetTypeAndShortName())
 			.ColorAndOpacity(this, &SReflectorTreeWidgetItem::GetTint)
 		];
 	}
@@ -79,7 +83,7 @@ TSharedRef<SWidget> SReflectorTreeWidgetItem::GenerateWidgetForColumn(const FNam
 			.Padding(FMargin(2.0f, 0.0f))
 			[
 				SNew(SHyperlink)
-				.Text(this, &SReflectorTreeWidgetItem::GetReadableLocationAsText)
+				.Text(WidgetInfo->GetWidgetReadableLocation())
 				.OnNavigate(this, &SReflectorTreeWidgetItem::HandleHyperlinkNavigate)
 			];
 	}
@@ -91,21 +95,25 @@ TSharedRef<SWidget> SReflectorTreeWidgetItem::GenerateWidgetForColumn(const FNam
 			.Padding(FMargin(2.0f, 0.0f))
 			[
 				SNew(STextBlock)
-				.Text(this, &SReflectorTreeWidgetItem::GetVisibilityAsString)
+				.Text(WidgetInfo->GetWidgetVisibilityText())
 					.Justification(ETextJustify::Center)
 			];
 	}
 	else if (ColumnName == NAME_Focusable)
 	{
-		return SNew(SBox)
-			.HAlign(HAlign_Center)
-			.VAlign(VAlign_Center)
-			.Padding(FMargin(2.0f, 0.0f))
-			[
-				SNew(SCheckBox)
-				.Style(FCoreStyle::Get(), TEXT("WidgetReflector.FocusableCheck"))
-				.IsChecked(this, &SReflectorTreeWidgetItem::GetFocusableAsCheckBoxState)
-			];
+		return BuildCheckBox(WidgetInfo->GetWidgetFocusable());
+	}
+	else if (ColumnName == NAME_Enabled)
+	{
+		return BuildCheckBox(WidgetInfo->GetWidgetEnabled());
+	}
+	else if (ColumnName == NAME_Volatile)
+	{
+		return BuildCheckBox(WidgetInfo->GetWidgetIsVolatile());
+	}
+	else if (ColumnName == NAME_HasActiveTimer)
+	{
+		return BuildCheckBox(WidgetInfo->GetWidgetHasActiveTimers());
 	}
 	else if ( ColumnName == NAME_Clipping )
 	{
@@ -115,7 +123,7 @@ TSharedRef<SWidget> SReflectorTreeWidgetItem::GenerateWidgetForColumn(const FNam
 			.Padding(FMargin(2.0f, 0.0f))
 			[
 				SNew(STextBlock)
-				.Text(this, &SReflectorTreeWidgetItem::GetClippingAsString)
+				.Text(WidgetInfo->GetWidgetClippingText())
 			];
 	}
 	else if (ColumnName == NAME_ForegroundColor )
@@ -136,6 +144,11 @@ TSharedRef<SWidget> SReflectorTreeWidgetItem::GenerateWidgetForColumn(const FNam
 					.Size(FVector2D(16.0f, 16.0f))
 			];
 	}
+	else if (ColumnName == NAME_ActualSize)
+	{
+		return SNew(STextBlock)
+			.Text(FText::FromString(WidgetInfo->GetLocalSize().ToString()));
+	}
 	else if (ColumnName == NAME_Address )
 	{
 		const FString WidgetAddress = FWidgetReflectorNodeUtils::WidgetAddressToString(WidgetInfo->GetWidgetAddress());
@@ -173,6 +186,7 @@ END_SLATE_FUNCTION_BUILD_OPTIMIZATION
 
 void SReflectorTreeWidgetItem::HandleHyperlinkNavigate()
 {
+	FAssetData CachedAssetData = WidgetInfo->GetWidgetAssetData();
 	if (CachedAssetData.IsValid())
 	{
 		if (OnAccessAsset.IsBound())
@@ -185,7 +199,7 @@ void SReflectorTreeWidgetItem::HandleHyperlinkNavigate()
 
 	if (OnAccessSourceCode.IsBound())
 	{
-		OnAccessSourceCode.Execute(GetWidgetFile(), GetWidgetLineNumber(), 0);
+		OnAccessSourceCode.Execute(WidgetInfo->GetWidgetFile(), WidgetInfo->GetWidgetLineNumber(), 0);
 	}
 }
 
diff --git a/Engine/Source/Developer/SlateReflector/Private/Widgets/SWidgetReflectorTreeWidgetItem.h b/Engine/Source/Developer/SlateReflector/Private/Widgets/SWidgetReflectorTreeWidgetItem.h
--- a/Engine/Source/Developer/SlateReflector/Private/Widgets/SWidgetReflectorTreeWidgetItem.h
+++ b/Engine/Source/Developer/SlateReflector/Private/Widgets/SWidgetReflectorTreeWidgetItem.h
@@ -23,9 +23,13 @@ public:
 	static FName NAME_WidgetInfo;
 	static FName NAME_Visibility;
 	static FName NAME_Focusable;
+	static FName NAME_Enabled;
+	static FName NAME_Volatile;
+	static FName NAME_HasActiveTimer;
 	static FName NAME_Clipping;
 	static FName NAME_ForegroundColor;
 	static FName NAME_Address;
+	static FName NAME_ActualSize;
 
 	SLATE_BEGIN_ARGS(SReflectorTreeWidgetItem)
 		: _WidgetInfoToVisualize()
@@ -55,53 +59,6 @@ public:
 
 protected:
 
-	/** @return String representation of the widget we are visualizing */
-	FText GetWidgetType() const
-	{
-		return CachedWidgetType;
-	}
-
-	FText GetWidgetTypeAndShortName() const
-	{
-		return CachedWidgetTypeAndShortName;
-	}
-
-	virtual FString GetReadableLocation() const override
-	{
-		return CachedReadableLocation.ToString();
-	}
-
-	FText GetReadableLocationAsText() const
-	{
-		return CachedReadableLocation;
-	}
-
-	FString GetWidgetFile() const
-	{
-		return CachedWidgetFile;
-	}
-
-	int32 GetWidgetLineNumber() const
-	{
-		return CachedWidgetLineNumber;
-	}
-
-	FText GetVisibilityAsString() const
-	{
-		return CachedWidgetVisibility;
-	}
-	
-	
-	FText GetClippingAsString() const
-	{
-		return CachedWidgetClipping;
-	}
-
-	ECheckBoxState GetFocusableAsCheckBoxState() const
-	{
-		return bCachedWidgetFocusable ? ECheckBoxState::Checked : ECheckBoxState::Unchecked;
-	}
-
 	/** @return The tint of the reflector node */
 	FSlateColor GetTint() const
 	{
@@ -115,22 +72,6 @@ private:
 	/** The info about the widget that we are visualizing. */
 	TSharedPtr<FWidgetReflectorNodeBase> WidgetInfo;
 
-	FText CachedWidgetType;
-	FText CachedWidgetTypeAndShortName;
-	FText CachedWidgetVisibility;
-	FText CachedWidgetClipping;
-	bool bCachedWidgetFocusable;
-	bool bCachedWidgetVisible;
-	bool bCachedWidgetNeedsTick;
-	bool bCachedWidgetIsVolatile;
-	bool bCachedWidgetIsVolatileIndirectly;
-	bool bCachedWidgetHasActiveTimers;
-	FText CachedReadableLocation;
-	FString CachedWidgetFile;
-	int32 CachedWidgetLineNumber;
-	FAssetData CachedAssetData;
-
 	FAccessSourceCode OnAccessSourceCode;
-
 	FAccessAsset OnAccessAsset;
 };
