Fix implicit type conversion with CreateWithImageData.

diff --git a/Engine/Source/Developer/SlateReflector/Private/Widgets/SWidgetSnapshotVisualizer.cpp b/Engine/Source/Developer/SlateReflector/Private/Widgets/SWidgetSnapshotVisualizer.cpp
--- a/Engine/Source/Developer/SlateReflector/Private/Widgets/SWidgetSnapshotVisualizer.cpp
+++ b/Engine/Source/Developer/SlateReflector/Private/Widgets/SWidgetSnapshotVisualizer.cpp
@@ -601,8 +601,8 @@ void FWidgetSnapshotData::LoadSnapshotFromJson(const TSharedRef<FJsonObject>& In
 				const TArray<TSharedPtr<FJsonValue>>& StructJsonArray = TextureDataJsonObject->GetArrayField(TEXT("Dimensions"));
 				check(StructJsonArray.Num() == 2);
 
-				TextureData.Dimensions.X = FMath::TruncToInt(StructJsonArray[0]->AsNumber());
-				TextureData.Dimensions.Y = FMath::TruncToInt(StructJsonArray[1]->AsNumber());
+				TextureData.Dimensions.X = (int32)(StructJsonArray[0]->AsNumber());
+				TextureData.Dimensions.Y = (int32)(StructJsonArray[1]->AsNumber());
 			}
 
 			{
@@ -613,7 +613,7 @@ void FWidgetSnapshotData::LoadSnapshotFromJson(const TSharedRef<FJsonObject>& In
 				const bool bIsCompressed = TextureDataJsonObject->GetBoolField(TEXT("IsCompressed"));
 				if (bIsCompressed)
 				{
-					const int32 UncompressedDataSizeBytes = FMath::TruncToInt(TextureDataJsonObject->GetNumberField(TEXT("UncompressedSize")));
+					const int32 UncompressedDataSizeBytes = (int32)(TextureDataJsonObject->GetNumberField(TEXT("UncompressedSize")));
 					TextureData.ColorData.AddZeroed(UncompressedDataSizeBytes / sizeof(FColor));
 
 					FCompression::UncompressMemory(NAME_Zlib, TextureData.ColorData.GetData(), UncompressedDataSizeBytes, DecodedTextureDataBytes.GetData(), DecodedTextureDataBytes.Num());
@@ -683,7 +683,7 @@ void FWidgetSnapshotData::CreateBrushes()
 	WindowTextureBrushes.Reserve(WindowTextureData.Num());
 
 	static int32 TextureIndex = 0;
-	for (const auto& TextureData : WindowTextureData)
+	for (const FWidgetSnapshotTextureData& TextureData : WindowTextureData)
 	{
 		if (TextureData.ColorData.Num() > 0)
 		{
@@ -699,7 +699,7 @@ void FWidgetSnapshotData::CreateBrushes()
 
 			WindowTextureBrushes.Add(FSlateDynamicImageBrush::CreateWithImageData(
 				*FString::Printf(TEXT("FWidgetSnapshotData_WindowTextureBrush_%d"), TextureIndex++), 
-				FVector2D(TextureData.Dimensions.X, TextureData.Dimensions.Y), 
+				FVector2D((float)TextureData.Dimensions.X, (float)TextureData.Dimensions.Y),
 				TextureDataAsBGRABytes
 				));
 		}
diff --git a/Engine/Source/Editor/AddContentDialog/Private/ViewModels/ContentSourceViewModel.cpp b/Engine/Source/Editor/AddContentDialog/Private/ViewModels/ContentSourceViewModel.cpp
--- a/Engine/Source/Editor/AddContentDialog/Private/ViewModels/ContentSourceViewModel.cpp
+++ b/Engine/Source/Editor/AddContentDialog/Private/ViewModels/ContentSourceViewModel.cpp
@@ -122,7 +122,7 @@ TSharedPtr<FSlateDynamicImageBrush> FContentSourceViewModel::CreateBrushFromRawD
 	if (bSucceeded)
 	{
 		FString UniqueResourceName = ResourceNamePrefix + "_" + FString::FromInt(ImageID++);
-		Brush = FSlateDynamicImageBrush::CreateWithImageData(FName(*UniqueResourceName), FVector2D(ImageWrapper->GetWidth(), ImageWrapper->GetHeight()), DecodedImage);
+		Brush = FSlateDynamicImageBrush::CreateWithImageData(FName(*UniqueResourceName), FVector2D((float)ImageWrapper->GetWidth(), (float)ImageWrapper->GetHeight()), DecodedImage);
 	}
 
 	return Brush;
diff --git a/Engine/Source/Runtime/Online/ImageDownload/Private/WebImage.cpp b/Engine/Source/Runtime/Online/ImageDownload/Private/WebImage.cpp
--- a/Engine/Source/Runtime/Online/ImageDownload/Private/WebImage.cpp
+++ b/Engine/Source/Runtime/Online/ImageDownload/Private/WebImage.cpp
@@ -147,7 +147,7 @@ bool FWebImage::ProcessHttpResponse(const FString& RequestUrl, FHttpResponsePtr
 
 	// make a dynamic image
 	FName ResourceName(*RequestUrl);
-	DownloadedBrush = FSlateDynamicImageBrush::CreateWithImageData(ResourceName, FVector2D(ImageWrapper->GetWidth(), ImageWrapper->GetHeight()), RawImageData);
+	DownloadedBrush = FSlateDynamicImageBrush::CreateWithImageData(ResourceName, FVector2D((float)ImageWrapper->GetWidth(), (float)ImageWrapper->GetHeight()), RawImageData);
 	return DownloadedBrush.IsValid();
 }
 
