WidgetReflector: Add command to take a snapshot.

diff --git a/Engine/Source/Developer/SlateReflector/Private/Widgets/SWidgetReflector.cpp b/Engine/Source/Developer/SlateReflector/Private/Widgets/SWidgetReflector.cpp
--- a/Engine/Source/Developer/SlateReflector/Private/Widgets/SWidgetReflector.cpp
+++ b/Engine/Source/Developer/SlateReflector/Private/Widgets/SWidgetReflector.cpp
@@ -73,6 +73,14 @@
 namespace WidgetReflectorImpl
 {
 
+/** Command to take a snapshot. */
+void TakeSnapshotCommand();
+
+FAutoConsoleCommand CCommandWidgetReflectorTakeSnapshot(
+	TEXT("WidgetReflector.TakeSnapshot")
+	, TEXT("Take a snapshot and save the result on the local drive.")
+	, FConsoleCommandDelegate::CreateStatic(&TakeSnapshotCommand));
+
 /** Information about a potential widget snapshot target */
 struct FWidgetSnapshotTarget
 {
@@ -2008,6 +2016,16 @@ void SWidgetReflector::HandleStartTreeWithUMG()
 	ReflectorTree->RequestTreeRefresh();
 }
 
+// command line
+
+void TakeSnapshotCommand()
+{
+	FWidgetSnapshotData SnapshotData;
+	SnapshotData.TakeSnapshot(false);
+	FString Filename = FPaths::CreateTempFilename(*FPaths::GameAgnosticSavedDir(), TEXT(""), TEXT(".widgetsnapshot"));
+	SnapshotData.SaveSnapshotToFile(Filename);
+}
+
 } // namespace WidgetReflectorImpl
 
 TSharedRef<SWidgetReflector> SWidgetReflector::New()
diff --git a/Engine/Source/Editor/UnrealEd/Private/EditorServer.cpp b/Engine/Source/Editor/UnrealEd/Private/EditorServer.cpp
--- a/Engine/Source/Editor/UnrealEd/Private/EditorServer.cpp
+++ b/Engine/Source/Editor/UnrealEd/Private/EditorServer.cpp
@@ -5705,15 +5705,6 @@ bool UEditorEngine::Exec( UWorld* InWorld, const TCHAR* Stream, FOutputDevice& A
 		MainFrameModule.RequestCloseEditor();
 		return true;
 	}
-	else if( FParse::Command(&Str,TEXT("WIDGETREFLECTOR")) )
-	{
-		if(!IsRunningCommandlet())
-		{
-			static const FName SlateReflectorModuleName("SlateReflector");
-			FModuleManager::LoadModuleChecked<ISlateReflectorModule>(SlateReflectorModuleName).DisplayWidgetReflector();
-		}
-		return true;
-	}
 	//----------------------------------------------------------------------------------
 	// LIGHTMASSDEBUG - Toggles whether UnrealLightmass.exe is launched automatically (default),
 	// or must be launched manually (e.g. through a debugger) with the -debug command line parameter.
