WidgetReflector: Fix the snapshot save to Json. It had missing properties and navigation.

diff --git a/Engine/Source/Developer/SlateReflector/Private/Models/NavigationSimulationNode.cpp b/Engine/Source/Developer/SlateReflector/Private/Models/NavigationSimulationNode.cpp
--- a/Engine/Source/Developer/SlateReflector/Private/Models/NavigationSimulationNode.cpp
+++ b/Engine/Source/Developer/SlateReflector/Private/Models/NavigationSimulationNode.cpp
@@ -2,6 +2,8 @@
 
 #include "Models/NavigationSimulationNode.h"
 
+#include "Dom/JsonObject.h"
+#include "Dom/JsonValue.h"
 #include "Layout/WidgetPath.h"
 #include "Models/WidgetReflectorNode.h"
 #include "Widgets/SWidget.h"
@@ -102,6 +104,241 @@ namespace NavigationSimulationNodeUtilsInternal
 	}
 }
 
+TSharedRef<FJsonValue> FNavigationSimulationWidgetNode::ToJson(const FNavigationSimulationWidgetNode& Node)
+{
+	struct Internal
+	{
+		static FString ConvertPtrIntToString(FNavigationSimulationWidgetInfo::TPointerAsInt Value)
+		{
+			return FWidgetReflectorNodeUtils::WidgetAddressToString(Value);
+		}
+		static TSharedRef<FJsonValue> CreateJsonValue(const FVector2D& InVec2D)
+		{
+			TArray<TSharedPtr<FJsonValue>> StructJsonArray;
+			StructJsonArray.Add(MakeShared<FJsonValueNumber>(InVec2D.X));
+			StructJsonArray.Add(MakeShared<FJsonValueNumber>(InVec2D.Y));
+			return MakeShared<FJsonValueArray>(StructJsonArray);
+		}
+		static TSharedRef<FJsonValue> CreateJsonValue(const FMatrix2x2& InMatrix)
+		{
+			float m00, m01, m10, m11;
+			InMatrix.GetMatrix(m00, m01, m10, m11);
+
+			TArray<TSharedPtr<FJsonValue>> StructJsonArray;
+			StructJsonArray.Add(MakeShared<FJsonValueNumber>(m00));
+			StructJsonArray.Add(MakeShared<FJsonValueNumber>(m01));
+			StructJsonArray.Add(MakeShared<FJsonValueNumber>(m10));
+			StructJsonArray.Add(MakeShared<FJsonValueNumber>(m11));
+			return MakeShared<FJsonValueArray>(StructJsonArray);
+		}
+		static TSharedRef<FJsonValue> CreateJsonValue(const FSlateLayoutTransform& InLayoutTransform)
+		{
+			TSharedRef<FJsonObject> StructJsonObject = MakeShared<FJsonObject>();
+			StructJsonObject->SetNumberField(TEXT("Scale"), InLayoutTransform.GetScale());
+			StructJsonObject->SetField(TEXT("Translation"), CreateJsonValue(InLayoutTransform.GetTranslation()));
+			return MakeShared<FJsonValueObject>(StructJsonObject);
+		}
+		static TSharedRef<FJsonValue> CreateJsonValue(const FSlateRenderTransform& InRenderTransform)
+		{
+			TSharedRef<FJsonObject> StructJsonObject = MakeShared<FJsonObject>();
+			StructJsonObject->SetField(TEXT("Matrix"), CreateJsonValue(InRenderTransform.GetMatrix()));
+			StructJsonObject->SetField(TEXT("Translation"), CreateJsonValue(InRenderTransform.GetTranslation()));
+			return MakeShared<FJsonValueObject>(StructJsonObject);
+		}
+		static TSharedRef<FJsonValue> CreateJsonValue(const FGeometry& InGeo)
+		{
+			TSharedRef<FJsonObject> StructJsonObject = MakeShared<FJsonObject>();
+			StructJsonObject->SetField(TEXT("AccumulatedRenderTransform"), CreateJsonValue(InGeo.GetAccumulatedRenderTransform()));
+			StructJsonObject->SetField(TEXT("AccumulatedLayoutTransform"), CreateJsonValue(InGeo.GetAccumulatedLayoutTransform()));
+			StructJsonObject->SetField(TEXT("LocalSize"), CreateJsonValue(InGeo.GetLocalSize()));
+			return MakeShared<FJsonValueObject>(StructJsonObject);
+		}
+		static TSharedRef<FJsonValue> CreateJsonValue(const FNavigationSimulationWidgetInfo& Value)
+		{
+			TSharedRef<FJsonObject> StructJsonObject = MakeShared<FJsonObject>();
+			StructJsonObject->SetStringField(TEXT("WidgetPtr"), ConvertPtrIntToString(Value.WidgetPtr));
+			StructJsonObject->SetStringField(TEXT("WidgetTypeAndShortName"), Value.WidgetTypeAndShortName.ToString());
+			StructJsonObject->SetBoolField(TEXT("bHasGeometry"), Value.bHasGeometry);
+			if (Value.bHasGeometry)
+			{
+				StructJsonObject->SetField(TEXT("WidgetGeometry"), CreateJsonValue(Value.WidgetGeometry));
+			}
+			return MakeShared<FJsonValueObject>(StructJsonObject);
+		}
+		static TSharedRef<FJsonValue> CreateJsonValue(const FNavigationSimulationWidgetNodeItem& Value)
+		{
+			TSharedRef<FJsonObject> StructJsonObject = MakeShared<FJsonObject>();
+			StructJsonObject->SetNumberField(TEXT("NavigationType"), (int32)Value.NavigationType);
+			StructJsonObject->SetNumberField(TEXT("ReplyBoundaryRule"), (int32)Value.ReplyBoundaryRule);
+			StructJsonObject->SetNumberField(TEXT("RoutedReason"), (int32)Value.RoutedReason);
+			StructJsonObject->SetField(TEXT("Destination"), CreateJsonValue(Value.Destination));
+			StructJsonObject->SetField(TEXT("ReplyEventHandler"), CreateJsonValue(Value.ReplyEventHandler));
+			StructJsonObject->SetField(TEXT("ReplyFocusRecipient"), CreateJsonValue(Value.ReplyFocusRecipient));
+			StructJsonObject->SetField(TEXT("WidgetThatShouldReceivedFocus"), CreateJsonValue(Value.WidgetThatShouldReceivedFocus));
+			StructJsonObject->SetField(TEXT("FocusedWidget"), CreateJsonValue(Value.FocusedWidget));
+			StructJsonObject->SetBoolField(TEXT("bIsDynamic"), Value.bIsDynamic);
+			StructJsonObject->SetBoolField(TEXT("bAlwaysHandleNavigationAttempt"), Value.bAlwaysHandleNavigationAttempt);
+			StructJsonObject->SetBoolField(TEXT("bCanFindWidgetForSetFocus"), Value.bCanFindWidgetForSetFocus);
+			StructJsonObject->SetBoolField(TEXT("bRoutedHandlerHasNavigationMeta"), Value.bRoutedHandlerHasNavigationMeta);
+			StructJsonObject->SetBoolField(TEXT("bHandledByViewport"), Value.bHandledByViewport);
+			return MakeShared<FJsonValueObject>(StructJsonObject);
+		}
+	};
+
+	TSharedRef<FJsonObject> RootJsonObject = MakeShared<FJsonObject>();
+	RootJsonObject->SetField(TEXT("NavigationSource"), Internal::CreateJsonValue(Node.NavigationSource));
+	RootJsonObject->SetNumberField(TEXT("NodeType"), (int32)Node.NodeType);
+	TArray<TSharedPtr<FJsonValue>> SimulationArray;
+	for (const FNavigationSimulationWidgetNodeItem& Item : Node.Simulations)
+	{
+		SimulationArray.Add(Internal::CreateJsonValue(Item));
+	}
+	RootJsonObject->SetArrayField(TEXT("Simulations"), SimulationArray);
+	return MakeShared<FJsonValueObject>(RootJsonObject);
+}
+
+FNavigationSimulationWidgetNodePtr FNavigationSimulationWidgetNode::FromJson(const TSharedRef<FJsonValue>& RootJsonValue)
+{
+	struct Internal
+	{
+		static void ParseJsonValue(const TSharedPtr<FJsonValue>& InJsonValue, FVector2D& Out)
+		{
+			if (!InJsonValue.IsValid())
+			{
+				return;
+			}
+			const TArray<TSharedPtr<FJsonValue>>& StructJsonArray = InJsonValue->AsArray();
+			check(StructJsonArray.Num() == 2);
+
+			Out = FVector2D(StructJsonArray[0]->AsNumber(), StructJsonArray[1]->AsNumber());
+		}
+		static FNavigationSimulationWidgetInfo::TPointerAsInt ParsePtrIntFromString(const FString& Value)
+		{
+			FWidgetReflectorNodeBase::TPointerAsInt Result = 0;
+			LexFromString(Result, *Value);
+			return Result;
+		}
+		static void ParseJsonValue(const TSharedPtr<FJsonValue>& InJsonValue, FMatrix2x2& Out)
+		{
+			if (!InJsonValue.IsValid())
+			{
+				return;
+			}
+			const TArray<TSharedPtr<FJsonValue>>& StructJsonArray = InJsonValue->AsArray();
+			check(StructJsonArray.Num() == 4);
+
+			Out = FMatrix2x2(
+				StructJsonArray[0]->AsNumber(),
+				StructJsonArray[1]->AsNumber(),
+				StructJsonArray[2]->AsNumber(),
+				StructJsonArray[3]->AsNumber()
+			);
+		}
+		static void ParseJsonValue(const TSharedPtr<FJsonValue>& InJsonValue, FSlateLayoutTransform& Out)
+		{
+			if (!InJsonValue.IsValid())
+			{
+				return;
+			}
+			const TSharedPtr<FJsonObject>& StructJsonObject = InJsonValue->AsObject();
+			check(StructJsonObject.IsValid());
+
+
+			float Scale = StructJsonObject->GetNumberField(TEXT("Scale"));
+			FVector2D Translation;
+			ParseJsonValue(StructJsonObject->GetField<EJson::None>(TEXT("Translation")), Translation);
+			Out = FSlateLayoutTransform(Scale, Translation);
+		}
+		static void ParseJsonValue(const TSharedPtr<FJsonValue>& InJsonValue, FSlateRenderTransform& Out)
+		{
+			if (!InJsonValue.IsValid())
+			{
+				return;
+			}
+			const TSharedPtr<FJsonObject>& StructJsonObject = InJsonValue->AsObject();
+			check(StructJsonObject.IsValid());
+
+			FMatrix2x2 Matrix;
+			ParseJsonValue(StructJsonObject->GetField<EJson::None>(TEXT("Matrix")), Matrix);
+			FVector2D Translation;
+			ParseJsonValue(StructJsonObject->GetField<EJson::None>(TEXT("Translation")), Translation);
+			Out = FSlateRenderTransform(Matrix, Translation);
+		}
+		static void ParseJsonValue(const TSharedPtr<FJsonValue>& InJsonValue, FGeometry& Out)
+		{
+			if (!InJsonValue.IsValid())
+			{
+				return;
+			}
+			const TSharedPtr<FJsonObject>& StructJsonObject = InJsonValue->AsObject();
+			check(StructJsonObject.IsValid());
+
+			FSlateLayoutTransform LayoutTransform;
+			FSlateRenderTransform RenderTransform;
+			ParseJsonValue(StructJsonObject->GetField<EJson::None>(TEXT("AccumulatedRenderTransform")), RenderTransform);
+			ParseJsonValue(StructJsonObject->GetField<EJson::None>(TEXT("AccumulatedLayoutTransform")), LayoutTransform);
+			FVector2D LocalSize;
+			ParseJsonValue(StructJsonObject->GetField<EJson::None>(TEXT("LocalSize")), LocalSize);
+
+			Out = FGeometry::MakeRoot(LocalSize, LayoutTransform, RenderTransform);
+		}
+		static void ParseJsonValue(const TSharedPtr<FJsonValue>& InJsonValue, FNavigationSimulationWidgetInfo& Out)
+		{
+			if (!InJsonValue.IsValid())
+			{
+				return;
+			}
+			const TSharedPtr<FJsonObject>& StructJsonObject = InJsonValue->AsObject();
+			check(StructJsonObject.IsValid());
+
+			Out.WidgetPtr = ParsePtrIntFromString(StructJsonObject->GetStringField(TEXT("WidgetPtr")));
+			Out.WidgetTypeAndShortName = FText::FromString(StructJsonObject->GetStringField(TEXT("WidgetTypeAndShortName")));
+			Out.bHasGeometry = StructJsonObject->GetBoolField(TEXT("bHasGeometry"));
+			if (Out.bHasGeometry)
+			{
+				ParseJsonValue(StructJsonObject->GetField<EJson::None>(TEXT("WidgetGeometry")), Out.WidgetGeometry);
+			}
+		}
+		static void ParseJsonValue(const TSharedPtr<FJsonValue>& InJsonValue, FNavigationSimulationWidgetNodeItem& Out)
+		{
+			if (!InJsonValue.IsValid())
+			{
+				return;
+			}
+			const TSharedPtr<FJsonObject>& StructJsonObject = InJsonValue->AsObject();
+			check(StructJsonObject.IsValid());
+
+			Out.NavigationType = (EUINavigation)(int32)StructJsonObject->GetNumberField(TEXT("NavigationType"));
+			Out.ReplyBoundaryRule = (EUINavigationRule)(int32)StructJsonObject->GetNumberField(TEXT("ReplyBoundaryRule"));
+			Out.RoutedReason = (FSlateNavigationEventSimulator::ERoutedReason)(int32)StructJsonObject->GetNumberField(TEXT("RoutedReason"));
+			ParseJsonValue(StructJsonObject->GetField<EJson::None>(TEXT("Destination")), Out.Destination);
+			ParseJsonValue(StructJsonObject->GetField<EJson::None>(TEXT("ReplyEventHandler")), Out.ReplyEventHandler);
+			ParseJsonValue(StructJsonObject->GetField<EJson::None>(TEXT("ReplyFocusRecipient")), Out.ReplyFocusRecipient);
+			ParseJsonValue(StructJsonObject->GetField<EJson::None>(TEXT("WidgetThatShouldReceivedFocus")), Out.WidgetThatShouldReceivedFocus);
+			ParseJsonValue(StructJsonObject->GetField<EJson::None>(TEXT("FocusedWidget")), Out.FocusedWidget);
+			Out.bIsDynamic = StructJsonObject->GetBoolField(TEXT("bIsDynamic"));
+			Out.bAlwaysHandleNavigationAttempt = StructJsonObject->GetBoolField(TEXT("bAlwaysHandleNavigationAttempt"));
+			Out.bCanFindWidgetForSetFocus = StructJsonObject->GetBoolField(TEXT("bCanFindWidgetForSetFocus"));
+			Out.bRoutedHandlerHasNavigationMeta = StructJsonObject->GetBoolField(TEXT("bRoutedHandlerHasNavigationMeta"));
+			Out.bHandledByViewport = StructJsonObject->GetBoolField(TEXT("bHandledByViewport"));
+		}
+	};
+
+	TSharedRef<FNavigationSimulationWidgetNode> Result = MakeShared<FNavigationSimulationWidgetNode>();
+	const TSharedPtr<FJsonObject>& RootJsonObject = RootJsonValue->AsObject();
+	check(RootJsonObject.IsValid());
+	Internal::ParseJsonValue(RootJsonObject->GetField<EJson::None>(TEXT("NavigationSource")), Result->NavigationSource);
+	Result->NodeType = (ENavigationSimulationNodeType)((int32)RootJsonObject->GetNumberField(TEXT("NodeType")));
+	const TArray<TSharedPtr<FJsonValue>>& SimulationsJsonArray = RootJsonObject->GetArrayField(TEXT("Simulations"));
+	for (const TSharedPtr<FJsonValue>& SimulationsJsonValue : SimulationsJsonArray)
+	{
+		FNavigationSimulationWidgetNodeItem NodeItem;
+		Internal::ParseJsonValue(SimulationsJsonValue.ToSharedRef(), NodeItem);
+		Result->Simulations.Add(MoveTemp(NodeItem));
+	}
+	return Result;
+}
+
 TArray<FNavigationSimulationWidgetNodePtr> FNavigationSimulationNodeUtils::BuildNavigationSimulationNodeListForLive(const TArray<FSlateNavigationEventSimulator::FSimulationResult>& SimulationResult)
 {
 	return NavigationSimulationNodeUtilsInternal::BuildNavigationSimulationNodeList(ENavigationSimulationNodeType::Live, SimulationResult);
diff --git a/Engine/Source/Developer/SlateReflector/Private/Models/NavigationSimulationNode.h b/Engine/Source/Developer/SlateReflector/Private/Models/NavigationSimulationNode.h
--- a/Engine/Source/Developer/SlateReflector/Private/Models/NavigationSimulationNode.h
+++ b/Engine/Source/Developer/SlateReflector/Private/Models/NavigationSimulationNode.h
@@ -9,6 +9,7 @@
 class SWidget;
 class FWidgetPath;
 class FWidgetReflectorNodeBase;
+class FJsonValue;
 
 struct FNavigationSimulationWidgetInfo;
 struct FNavigationSimulationWidgetNodeItem;
@@ -26,8 +27,9 @@ struct FNavigationSimulationWidgetInfo
 	// Should be PTRINT but if the platform that produce the file is 64bits and the reader is not then we could have a issue
 	using TPointerAsInt = uint64;
 
-	FNavigationSimulationWidgetInfo(const TSharedPtr<const SWidget>& Widget);
-	FNavigationSimulationWidgetInfo(const FWidgetPath& WidgetPath);
+	FNavigationSimulationWidgetInfo() = default;
+	explicit FNavigationSimulationWidgetInfo(const TSharedPtr<const SWidget>& Widget);
+	explicit FNavigationSimulationWidgetInfo(const FWidgetPath& WidgetPath);
 
 	/** Is the widget set to null. */
 	bool IsWidgetExplicitlyNull() const { return WidgetPtr == 0; }
@@ -42,7 +44,8 @@ struct FNavigationSimulationWidgetInfo
 
 struct FNavigationSimulationWidgetNodeItem
 {
-	FNavigationSimulationWidgetNodeItem(const FSlateNavigationEventSimulator::FSimulationResult& Result);
+	FNavigationSimulationWidgetNodeItem() = default;
+	explicit FNavigationSimulationWidgetNodeItem(const FSlateNavigationEventSimulator::FSimulationResult& Result);
 
 	EUINavigation NavigationType;
 	FNavigationSimulationWidgetInfo Destination;
@@ -63,11 +66,18 @@ struct FNavigationSimulationWidgetNodeItem
 
 struct FNavigationSimulationWidgetNode
 {
-	FNavigationSimulationWidgetNode(ENavigationSimulationNodeType NodeType, const FWidgetPath& InNavigationSource);
+	FNavigationSimulationWidgetNode() = default;
+	explicit FNavigationSimulationWidgetNode(ENavigationSimulationNodeType NodeType, const FWidgetPath& InNavigationSource);
 
 	FNavigationSimulationWidgetInfo NavigationSource;
 	TArray<FNavigationSimulationWidgetNodeItem, TInlineAllocator<4>> Simulations;
 	ENavigationSimulationNodeType NodeType;
+
+	/** Save this node data as a JSON object */
+	static TSharedRef<FJsonValue> ToJson(const FNavigationSimulationWidgetNode& Node);
+
+	/** Populate this node data from a JSON object */
+	static FNavigationSimulationWidgetNodePtr FromJson(const TSharedRef<FJsonValue>& RootJsonValue);
 };
 
 
diff --git a/Engine/Source/Developer/SlateReflector/Private/Models/WidgetReflectorNode.cpp b/Engine/Source/Developer/SlateReflector/Private/Models/WidgetReflectorNode.cpp
--- a/Engine/Source/Developer/SlateReflector/Private/Models/WidgetReflectorNode.cpp
+++ b/Engine/Source/Developer/SlateReflector/Private/Models/WidgetReflectorNode.cpp
@@ -220,8 +220,8 @@ TSharedRef<FSnapshotWidgetReflectorNode> FSnapshotWidgetReflectorNode::Create(co
 }
 
 FSnapshotWidgetReflectorNode::FSnapshotWidgetReflectorNode()
-	: CachedWidgetLineNumber(0)
-	, CachedWidgetEnabled(false)
+	: bCachedWidgetEnabled(false)
+	, CachedWidgetLineNumber(0)
 {
 }
 
@@ -232,10 +232,11 @@ FSnapshotWidgetReflectorNode::FSnapshotWidgetReflectorNode(const FArrangedWidget
 	, CachedWidgetVisibilityText(FWidgetReflectorNodeUtils::GetWidgetVisibilityText(InArrangedWidget.Widget))
 	, bCachedWidgetVisible(FWidgetReflectorNodeUtils::GetWidgetVisibility(InArrangedWidget.Widget))
 	, bCachedWidgetFocusable(FWidgetReflectorNodeUtils::GetWidgetFocusable(InArrangedWidget.Widget))
-	, CachedWidgetNeedsTick(FWidgetReflectorNodeUtils::GetWidgetNeedsTick(InArrangedWidget.Widget))
-	, CachedWidgetIsVolatile(FWidgetReflectorNodeUtils::GetWidgetIsVolatile(InArrangedWidget.Widget))
-	, CachedWidgetIsVolatileIndirectly(FWidgetReflectorNodeUtils::GetWidgetIsVolatileIndirectly(InArrangedWidget.Widget))
-	, CachedWidgetHasActiveTimers(FWidgetReflectorNodeUtils::GetWidgetHasActiveTimers(InArrangedWidget.Widget))
+	, bCachedWidgetNeedsTick(FWidgetReflectorNodeUtils::GetWidgetNeedsTick(InArrangedWidget.Widget))
+	, bCachedWidgetIsVolatile(FWidgetReflectorNodeUtils::GetWidgetIsVolatile(InArrangedWidget.Widget))
+	, bCachedWidgetIsVolatileIndirectly(FWidgetReflectorNodeUtils::GetWidgetIsVolatileIndirectly(InArrangedWidget.Widget))
+	, bCachedWidgetHasActiveTimers(FWidgetReflectorNodeUtils::GetWidgetHasActiveTimers(InArrangedWidget.Widget))
+	, bCachedWidgetEnabled(FWidgetReflectorNodeUtils::GetWidgetEnabled(InArrangedWidget.Widget))
 	, CachedWidgetClippingText(FWidgetReflectorNodeUtils::GetWidgetClippingText(InArrangedWidget.Widget))
 	, CachedWidgetReadableLocation(FWidgetReflectorNodeUtils::GetWidgetReadableLocation(InArrangedWidget.Widget))
 	, CachedWidgetFile(FWidgetReflectorNodeUtils::GetWidgetFile(InArrangedWidget.Widget))
@@ -244,7 +245,6 @@ FSnapshotWidgetReflectorNode::FSnapshotWidgetReflectorNode(const FArrangedWidget
 	, CachedWidgetDesiredSize(FWidgetReflectorNodeUtils::GetWidgetDesiredSize(InArrangedWidget.Widget))
 	, CachedWidgetForegroundColor(FWidgetReflectorNodeUtils::GetWidgetForegroundColor(InArrangedWidget.Widget))
 	, CachedWidgetAddress(FWidgetReflectorNodeUtils::GetWidgetAddress(InArrangedWidget.Widget))
-	, CachedWidgetEnabled(FWidgetReflectorNodeUtils::GetWidgetEnabled(InArrangedWidget.Widget))
 {
 }
 
@@ -280,7 +280,7 @@ bool FSnapshotWidgetReflectorNode::GetWidgetFocusable() const
 
 bool FSnapshotWidgetReflectorNode::GetWidgetVisible() const
 {
-	return bCachedWidgetFocusable;
+	return bCachedWidgetVisible;
 }
 
 FText FSnapshotWidgetReflectorNode::GetWidgetClippingText() const
@@ -290,22 +290,22 @@ FText FSnapshotWidgetReflectorNode::GetWidgetClippingText() const
 
 bool FSnapshotWidgetReflectorNode::GetWidgetNeedsTick() const
 {
-	return CachedWidgetNeedsTick;
+	return bCachedWidgetNeedsTick;
 }
 
 bool FSnapshotWidgetReflectorNode::GetWidgetIsVolatile() const
 {
-	return CachedWidgetIsVolatile;
+	return bCachedWidgetIsVolatile;
 }
 
 bool FSnapshotWidgetReflectorNode::GetWidgetIsVolatileIndirectly() const
 {
-	return CachedWidgetIsVolatileIndirectly;
+	return bCachedWidgetIsVolatileIndirectly;
 }
 
 bool FSnapshotWidgetReflectorNode::GetWidgetHasActiveTimers() const
 {
-	return CachedWidgetHasActiveTimers;
+	return bCachedWidgetHasActiveTimers;
 }
 
 FText FSnapshotWidgetReflectorNode::GetWidgetReadableLocation() const
@@ -350,7 +350,7 @@ FWidgetReflectorNodeBase::TPointerAsInt FSnapshotWidgetReflectorNode::GetWidgetA
 
 bool FSnapshotWidgetReflectorNode::GetWidgetEnabled() const
 {
-	return CachedWidgetEnabled;
+	return bCachedWidgetEnabled;
 }
 
 TSharedRef<FJsonValue> FSnapshotWidgetReflectorNode::ToJson(const TSharedRef<FSnapshotWidgetReflectorNode>& RootSnapshotNode)
@@ -435,17 +435,25 @@ TSharedRef<FJsonValue> FSnapshotWidgetReflectorNode::ToJson(const TSharedRef<FSn
 	RootJsonObject->SetField(TEXT("AccumulatedRenderTransform"), Internal::CreateSlateRenderTransformJsonValue(RootSnapshotNode->GetAccumulatedRenderTransform()));
 	RootJsonObject->SetField(TEXT("LocalSize"), Internal::CreateVector2DJsonValue(RootSnapshotNode->GetLocalSize()));
 	RootJsonObject->SetField(TEXT("HitTestInfo"), Internal::CreateWidgetHitTestInfoJsonValue(RootSnapshotNode->HitTestInfo));
+	RootJsonObject->SetField(TEXT("Tint"), Internal::CreateLinearColorJsonValue(RootSnapshotNode->Tint));
 	RootJsonObject->SetStringField(TEXT("WidgetType"), RootSnapshotNode->CachedWidgetType.ToString());
+	RootJsonObject->SetStringField(TEXT("WidgetTypeAndShortName"), RootSnapshotNode->CachedWidgetTypeAndShortName.ToString());
 	RootJsonObject->SetStringField(TEXT("WidgetVisibilityText"), RootSnapshotNode->CachedWidgetVisibilityText.ToString());
+	RootJsonObject->SetBoolField(TEXT("WidgetVisible"), RootSnapshotNode->bCachedWidgetVisible);
 	RootJsonObject->SetBoolField(TEXT("WidgetFocusable"), RootSnapshotNode->bCachedWidgetFocusable);
+	RootJsonObject->SetBoolField(TEXT("WidgetNeedsTick"), RootSnapshotNode->bCachedWidgetNeedsTick);
+	RootJsonObject->SetBoolField(TEXT("WidgetIsVolatile"), RootSnapshotNode->bCachedWidgetIsVolatile);
+	RootJsonObject->SetBoolField(TEXT("WidgetIsVolatileIndirectly"), RootSnapshotNode->bCachedWidgetIsVolatileIndirectly);
+	RootJsonObject->SetBoolField(TEXT("WidgetHasActiveTimers"), RootSnapshotNode->bCachedWidgetHasActiveTimers);
+	RootJsonObject->SetBoolField(TEXT("WidgetEnabled"), RootSnapshotNode->bCachedWidgetEnabled);
+	RootJsonObject->SetStringField(TEXT("WidgetClippingText"), RootSnapshotNode->CachedWidgetClippingText.ToString());
 	RootJsonObject->SetStringField(TEXT("WidgetReadableLocation"), RootSnapshotNode->CachedWidgetReadableLocation.ToString());
 	RootJsonObject->SetStringField(TEXT("WidgetFile"), RootSnapshotNode->CachedWidgetFile);
 	RootJsonObject->SetNumberField(TEXT("WidgetLineNumber"), RootSnapshotNode->CachedWidgetLineNumber);
-	RootJsonObject->SetStringField(TEXT("WidgetAssetPath"), RootSnapshotNode->CachedWidgetAssetData.ObjectPath.ToString());
 	RootJsonObject->SetField(TEXT("WidgetDesiredSize"), Internal::CreateVector2DJsonValue(RootSnapshotNode->CachedWidgetDesiredSize));
 	RootJsonObject->SetField(TEXT("WidgetForegroundColor"), Internal::CreateSlateColorJsonValue(RootSnapshotNode->CachedWidgetForegroundColor));
 	RootJsonObject->SetStringField(TEXT("WidgetAddress"), Internal::ConvertPtrIntToString(RootSnapshotNode->CachedWidgetAddress));
-	RootJsonObject->SetBoolField(TEXT("WidgetEnabled"), RootSnapshotNode->CachedWidgetEnabled);
+	RootJsonObject->SetStringField(TEXT("WidgetAssetPath"), RootSnapshotNode->CachedWidgetAssetData.ObjectPath.ToString());
 
 	TArray<TSharedPtr<FJsonValue>> ChildNodesJsonArray;
 	for (const auto& ChildReflectorNode : RootSnapshotNode->ChildNodes)
@@ -464,7 +472,10 @@ TSharedRef<FSnapshotWidgetReflectorNode> FSnapshotWidgetReflectorNode::FromJson(
 	{
 		static FVector2D ParseVector2DJsonValue(const TSharedPtr<FJsonValue>& InJsonValue)
 		{
-			check(InJsonValue.IsValid());
+			if (!InJsonValue.IsValid())
+			{
+				return FVector2D::ZeroVector;
+			}
 			const TArray<TSharedPtr<FJsonValue>>& StructJsonArray = InJsonValue->AsArray();
 			check(StructJsonArray.Num() == 2);
 
@@ -476,7 +487,10 @@ TSharedRef<FSnapshotWidgetReflectorNode> FSnapshotWidgetReflectorNode::FromJson(
 
 		static FMatrix2x2 ParseMatrix2x2JsonValue(const TSharedPtr<FJsonValue>& InJsonValue)
 		{
-			check(InJsonValue.IsValid());
+			if (!InJsonValue.IsValid())
+			{
+				return FMatrix2x2();
+			}
 			const TArray<TSharedPtr<FJsonValue>>& StructJsonArray = InJsonValue->AsArray();
 			check(StructJsonArray.Num() == 4);
 
@@ -490,7 +504,10 @@ TSharedRef<FSnapshotWidgetReflectorNode> FSnapshotWidgetReflectorNode::FromJson(
 
 		static FSlateLayoutTransform ParseSlateLayoutTransformJsonValue(const TSharedPtr<FJsonValue>& InJsonValue)
 		{
-			check(InJsonValue.IsValid());
+			if (!InJsonValue.IsValid())
+			{
+				return FSlateLayoutTransform();
+			}
 			const TSharedPtr<FJsonObject>& StructJsonObject = InJsonValue->AsObject();
 			check(StructJsonObject.IsValid());
 
@@ -502,7 +519,10 @@ TSharedRef<FSnapshotWidgetReflectorNode> FSnapshotWidgetReflectorNode::FromJson(
 
 		static FSlateRenderTransform ParseSlateRenderTransformJsonValue(const TSharedPtr<FJsonValue>& InJsonValue)
 		{
-			check(InJsonValue.IsValid());
+			if (!InJsonValue.IsValid())
+			{
+				return FSlateRenderTransform();
+			}
 			const TSharedPtr<FJsonObject>& StructJsonObject = InJsonValue->AsObject();
 			check(StructJsonObject.IsValid());
 
@@ -514,7 +534,10 @@ TSharedRef<FSnapshotWidgetReflectorNode> FSnapshotWidgetReflectorNode::FromJson(
 
 		static FLinearColor ParseLinearColorJsonValue(const TSharedPtr<FJsonValue>& InJsonValue)
 		{
-			check(InJsonValue.IsValid());
+			if(!InJsonValue.IsValid())
+			{
+				return FLinearColor(EForceInit::ForceInit);
+			}
 			const TArray<TSharedPtr<FJsonValue>>& StructJsonArray = InJsonValue->AsArray();
 			check(StructJsonArray.Num() == 4);
 
@@ -528,7 +551,10 @@ TSharedRef<FSnapshotWidgetReflectorNode> FSnapshotWidgetReflectorNode::FromJson(
 
 		static FSlateColor ParseSlateColorJsonValue(const TSharedPtr<FJsonValue>& InJsonValue)
 		{
-			check(InJsonValue.IsValid());
+			if (!InJsonValue.IsValid())
+			{
+				return FSlateColor();
+			}
 			const TSharedPtr<FJsonObject>& StructJsonObject = InJsonValue->AsObject();
 			check(StructJsonObject.IsValid());
 
@@ -545,7 +571,10 @@ TSharedRef<FSnapshotWidgetReflectorNode> FSnapshotWidgetReflectorNode::FromJson(
 
 		static FWidgetHitTestInfo ParseWidgetHitTestInfoJsonValue(const TSharedPtr<FJsonValue>& InJsonValue)
 		{
-			check(InJsonValue.IsValid());
+			if (!InJsonValue.IsValid())
+			{
+				return FWidgetHitTestInfo();
+			}
 			const TSharedPtr<FJsonObject>& StructJsonObject = InJsonValue->AsObject();
 			check(StructJsonObject.IsValid());
 
@@ -572,18 +601,26 @@ TSharedRef<FSnapshotWidgetReflectorNode> FSnapshotWidgetReflectorNode::FromJson(
 	const FSlateRenderTransform RenderTransform = Internal::ParseSlateRenderTransformJsonValue(RootJsonObject->GetField<EJson::None>(TEXT("AccumulatedRenderTransform")));
 	const FVector2D LocalSize = Internal::ParseVector2DJsonValue(RootJsonObject->GetField<EJson::None>(TEXT("LocalSize")));
 	RootSnapshotNode->WidgetGeometry = FGeometry::MakeRoot(LocalSize, LayoutTransform, RenderTransform);
-
 	RootSnapshotNode->HitTestInfo = Internal::ParseWidgetHitTestInfoJsonValue(RootJsonObject->GetField<EJson::None>(TEXT("HitTestInfo")));
+	RootSnapshotNode->Tint = Internal::ParseLinearColorJsonValue(RootJsonObject->GetField<EJson::None>(TEXT("Tint")));
+
 	RootSnapshotNode->CachedWidgetType = FText::FromString(RootJsonObject->GetStringField(TEXT("WidgetType")));
+	RootSnapshotNode->CachedWidgetTypeAndShortName = FText::FromString(RootJsonObject->GetStringField(TEXT("WidgetTypeAndShortName")));
 	RootSnapshotNode->CachedWidgetVisibilityText = FText::FromString(RootJsonObject->GetStringField(TEXT("WidgetVisibilityText")));
+	RootSnapshotNode->bCachedWidgetVisible = RootJsonObject->GetBoolField(TEXT("WidgetVisible"));
 	RootSnapshotNode->bCachedWidgetFocusable = RootJsonObject->GetBoolField(TEXT("WidgetFocusable"));
+	RootSnapshotNode->bCachedWidgetNeedsTick = RootJsonObject->GetBoolField(TEXT("WidgetNeedsTick"));
+	RootSnapshotNode->bCachedWidgetIsVolatile = RootJsonObject->GetBoolField(TEXT("WidgetIsVolatile"));
+	RootSnapshotNode->bCachedWidgetIsVolatileIndirectly = RootJsonObject->GetBoolField(TEXT("WidgetIsVolatileIndirectly"));
+	RootSnapshotNode->bCachedWidgetHasActiveTimers = RootJsonObject->GetBoolField(TEXT("WidgetHasActiveTimers"));
+	RootSnapshotNode->bCachedWidgetEnabled = RootJsonObject->GetBoolField(TEXT("WidgetEnabled"));
+	RootSnapshotNode->CachedWidgetClippingText = FText::FromString(RootJsonObject->GetStringField(TEXT("WidgetClippingText")));
 	RootSnapshotNode->CachedWidgetReadableLocation = FText::FromString(RootJsonObject->GetStringField(TEXT("WidgetReadableLocation")));
 	RootSnapshotNode->CachedWidgetFile = RootJsonObject->GetStringField(TEXT("WidgetFile"));
 	RootSnapshotNode->CachedWidgetLineNumber = RootJsonObject->GetIntegerField(TEXT("WidgetLineNumber"));
 	RootSnapshotNode->CachedWidgetDesiredSize = Internal::ParseVector2DJsonValue(RootJsonObject->GetField<EJson::None>(TEXT("WidgetDesiredSize")));
 	RootSnapshotNode->CachedWidgetForegroundColor = Internal::ParseSlateColorJsonValue(RootJsonObject->GetField<EJson::None>(TEXT("WidgetForegroundColor")));
 	RootSnapshotNode->CachedWidgetAddress = Internal::ParsePtrIntFromString(RootJsonObject->GetStringField(TEXT("WidgetAddress")));
-	RootSnapshotNode->CachedWidgetEnabled = RootJsonObject->GetBoolField(TEXT("WidgetEnabled"));
 
 	FName AssetPath(*RootJsonObject->GetStringField(TEXT("WidgetAssetPath")));
 	IAssetRegistry& AssetRegistry = FModuleManager::LoadModuleChecked<FAssetRegistryModule>(TEXT("AssetRegistry")).Get();
diff --git a/Engine/Source/Developer/SlateReflector/Private/Models/WidgetReflectorNode.h b/Engine/Source/Developer/SlateReflector/Private/Models/WidgetReflectorNode.h
--- a/Engine/Source/Developer/SlateReflector/Private/Models/WidgetReflectorNode.h
+++ b/Engine/Source/Developer/SlateReflector/Private/Models/WidgetReflectorNode.h
@@ -28,8 +28,8 @@ enum class EWidgetReflectorNodeType : uint8
  */
 struct FWidgetHitTestInfo
 {
-	bool IsHitTestVisible;
-	bool AreChildrenHitTestVisible;
+	bool IsHitTestVisible = false;
+	bool AreChildrenHitTestVisible = false;
 };
 
 /** 
@@ -328,7 +328,7 @@ private:
 	/** The type string of the widget at the point it was passed to Initialize */
 	FText CachedWidgetType;
 
-	/**  */
+	/** The type and short name string of the widget at the point it was passed to Initialize */
 	FText CachedWidgetTypeAndShortName;
 
 	/** The visibility string of the widget at the point it was passed to Initialize */
@@ -340,10 +340,20 @@ private:
 	/** The focusability of the widget at the point it was passed to Initialize */
 	bool bCachedWidgetFocusable;
 
-	bool CachedWidgetNeedsTick;
-	bool CachedWidgetIsVolatile;
-	bool CachedWidgetIsVolatileIndirectly;
-	bool CachedWidgetHasActiveTimers;
+	/** The ticking state of the widget at the point it was passed to Initialize */
+	bool bCachedWidgetNeedsTick;
+
+	/** The volatility state of the widget at the point it was passed to Initialize */
+	bool bCachedWidgetIsVolatile;
+
+	/** The volatility indirectly state of the widget at the point it was passed to Initialize */
+	bool bCachedWidgetIsVolatileIndirectly;
+
+	/** The active timer state of the widget at the point it was passed to Initialize */
+	bool bCachedWidgetHasActiveTimers;
+
+	/** The enabled state of the widget at the point it was passed to Initialize */
+	bool bCachedWidgetEnabled;
 	
 	/** The clipping string of the widget at the point it was passed to Initialize */
 	FText CachedWidgetClippingText;
@@ -368,9 +378,6 @@ private:
 
 	/** The in-memory address of the widget at the point it was passed to Initialize */
 	TPointerAsInt CachedWidgetAddress;
-
-	/** The enabled state of the widget at the point it was passed to Initialize */
-	bool CachedWidgetEnabled;
 };
 
 
diff --git a/Engine/Source/Developer/SlateReflector/Private/Widgets/SWidgetReflector.cpp b/Engine/Source/Developer/SlateReflector/Private/Widgets/SWidgetReflector.cpp
--- a/Engine/Source/Developer/SlateReflector/Private/Widgets/SWidgetReflector.cpp
+++ b/Engine/Source/Developer/SlateReflector/Private/Widgets/SWidgetReflector.cpp
@@ -7,6 +7,7 @@
 #include "Rendering/DrawElements.h"
 #include "Misc/App.h"
 #include "Misc/ConfigCacheIni.h"
+#include "Misc/Parse.h"
 #include "Modules/ModuleManager.h"
 #include "Framework/Application/SlateApplication.h"
 #include "Framework/Docking/TabManager.h"
@@ -74,12 +75,12 @@ namespace WidgetReflectorImpl
 {
 
 /** Command to take a snapshot. */
-void TakeSnapshotCommand();
+void TakeSnapshotCommand(const TArray<FString>&);
 
 FAutoConsoleCommand CCommandWidgetReflectorTakeSnapshot(
 	TEXT("WidgetReflector.TakeSnapshot")
-	, TEXT("Take a snapshot and save the result on the local drive.")
-	, FConsoleCommandDelegate::CreateStatic(&TakeSnapshotCommand));
+	, TEXT("Take a snapshot and save the result on the local drive. ie. WidgetReflector.TakeSnapshot [Delay=1.0] [Navigation=false]")
+	, FConsoleCommandWithArgsDelegate::CreateStatic(&TakeSnapshotCommand));
 
 /** Information about a potential widget snapshot target */
 struct FWidgetSnapshotTarget
@@ -1924,7 +1925,7 @@ void SWidgetReflector::HandleReflectorTreeSelectionChanged( TSharedPtr<FWidgetRe
 		}
 	}
 
-	if (SelectedWidgetObjects.Num() > 0)
+	if (GIsEditor && SelectedWidgetObjects.Num() > 0)
 	{
 		TabManager->TryInvokeTab(WidgetReflectorTabID::WidgetDetails);
 		if (PropertyViewPtr.IsValid())
@@ -2016,14 +2017,45 @@ void SWidgetReflector::HandleStartTreeWithUMG()
 	ReflectorTree->RequestTreeRefresh();
 }
 
-// command line
+// console command
 
-void TakeSnapshotCommand()
+static FDelegateHandle TakeSnapshotEndFrameHandle;
+
+void TakeSnapshotCommand_EndFrame(double RequestedTime, bool bRequestNavigation)
 {
-	FWidgetSnapshotData SnapshotData;
-	SnapshotData.TakeSnapshot(false);
-	FString Filename = FPaths::CreateTempFilename(*FPaths::GameAgnosticSavedDir(), TEXT(""), TEXT(".widgetsnapshot"));
-	SnapshotData.SaveSnapshotToFile(Filename);
+	if (RequestedTime <= FApp::GetCurrentTime())
+	{
+		FWidgetSnapshotData SnapshotData;
+		SnapshotData.TakeSnapshot(bRequestNavigation);
+		FString Filename = FPaths::CreateTempFilename(*FPaths::GameAgnosticSavedDir(), TEXT(""), TEXT(".widgetsnapshot"));
+		SnapshotData.SaveSnapshotToFile(Filename);
+
+		FCoreDelegates::OnEndFrame.Remove(TakeSnapshotEndFrameHandle);
+		TakeSnapshotEndFrameHandle.Reset();
+	}
+}
+
+void TakeSnapshotCommand(const TArray<FString>& Args)
+{
+	FCoreDelegates::OnEndFrame.Remove(TakeSnapshotEndFrameHandle);
+
+	float RequestedDelay = 0.001f;
+	bool bRequestNavigation = false;
+	for (const FString& Arg : Args)
+	{
+		if (FParse::Value(*Arg, TEXT("Delay="), RequestedDelay))
+		{
+			continue;
+		}
+		if (FParse::Bool(*Arg, TEXT("Navigation="), bRequestNavigation))
+		{
+			continue;
+		}
+	}
+
+	const double CurrentTime = FApp::GetCurrentTime();
+	const double RequestedTimeDelay = CurrentTime + (double)RequestedDelay;
+	TakeSnapshotEndFrameHandle = FCoreDelegates::OnEndFrame.AddStatic(&TakeSnapshotCommand_EndFrame, RequestedTimeDelay, bRequestNavigation);
 }
 
 } // namespace WidgetReflectorImpl
diff --git a/Engine/Source/Developer/SlateReflector/Private/Widgets/SWidgetSnapshotVisualizer.cpp b/Engine/Source/Developer/SlateReflector/Private/Widgets/SWidgetSnapshotVisualizer.cpp
--- a/Engine/Source/Developer/SlateReflector/Private/Widgets/SWidgetSnapshotVisualizer.cpp
+++ b/Engine/Source/Developer/SlateReflector/Private/Widgets/SWidgetSnapshotVisualizer.cpp
@@ -385,7 +385,7 @@ void FWidgetSnapshotData::CreateSnapshot(const TArray<TSharedRef<SWindow>>& Visi
 	Reset();
 	Reserve(VisibleWindows.Num());
 
-	for (const auto& VisibleWindow : VisibleWindows)
+	for (const TSharedRef<SWindow>& VisibleWindow : VisibleWindows)
 	{
 		// Snapshot the current state of this window widget hierarchy
 		Windows.Add(FWidgetReflectorNodeUtils::NewSnapshotNodeTreeFrom(FArrangedWidget(VisibleWindow, VisibleWindow->GetWindowGeometryInScreen())));
@@ -446,15 +446,21 @@ void FWidgetSnapshotData::SaveSnapshotToBuffer(TArray<uint8>& OutData) const
 	BufferWriter.SerializeCompressed(TmpJsonData.GetData(), TmpJsonData.Num(), NAME_Zlib);
 }
 
+double SnapshotJsonVersion = 2.0;
+
 TSharedRef<FJsonObject> FWidgetSnapshotData::SaveSnapshotAsJson() const
 {
 	check(Windows.Num() == WindowTextureData.Num());
 
-	TSharedRef<FJsonObject> RootJsonObject = MakeShareable(new FJsonObject());
+	TSharedRef<FJsonObject> RootJsonObject = MakeShared<FJsonObject>();
+
+	{
+		RootJsonObject->SetNumberField(TEXT("Version"), SnapshotJsonVersion);
+	}
 
 	{
 		TArray<TSharedPtr<FJsonValue>> WindowsJsonArray;
-		for (const auto& Window : Windows)
+		for (const TSharedPtr<FWidgetReflectorNodeBase>& Window : Windows)
 		{
 			check(Window->GetNodeType() == EWidgetReflectorNodeType::Snapshot);
 			WindowsJsonArray.Add(FSnapshotWidgetReflectorNode::ToJson(StaticCastSharedRef<FSnapshotWidgetReflectorNode>(Window.ToSharedRef())));
@@ -462,6 +468,24 @@ TSharedRef<FJsonObject> FWidgetSnapshotData::SaveSnapshotAsJson() const
 		RootJsonObject->SetArrayField(TEXT("Windows"), WindowsJsonArray);
 	}
 
+	{
+		TArray<TSharedPtr<FJsonValue>> NavigationDataArray;
+		
+		for (const FWidgetSnapshotNavigationSimulationData& NavigationSimulation : NavigationSimulationData)
+		{
+			TSharedRef<FJsonObject> NavigationData = MakeShared<FJsonObject>();
+			TArray<TSharedPtr<FJsonValue>> SimulationDataArray;
+			for (const FNavigationSimulationWidgetNodePtr& SimulationData : NavigationSimulation.SimulationData)
+			{
+				check(SimulationData.Get());
+				SimulationDataArray.Add(FNavigationSimulationWidgetNode::ToJson(*SimulationData.Get()));
+			}
+			NavigationData->SetArrayField(TEXT("SimulationData"), SimulationDataArray);
+			NavigationDataArray.Add(MakeShared<FJsonValueObject>(NavigationData));
+		}
+		RootJsonObject->SetArrayField(TEXT("NavigationData"), NavigationDataArray);
+	}
+
 	{
 		TArray<TSharedPtr<FJsonValue>> TexturesJsonArray;
 		for (const FWidgetSnapshotTextureData& TextureData : WindowTextureData)
@@ -580,6 +604,14 @@ void FWidgetSnapshotData::LoadSnapshotFromJson(const TSharedRef<FJsonObject>& In
 {
 	Reset();
 
+	{
+		const double VersionNumber = InRootJsonObject->GetNumberField(TEXT("Version"));
+		if (VersionNumber < SnapshotJsonVersion)
+		{
+			UE_LOG(LogSlate, Error, TEXT("The version of the snapshot (%f) is older than the current version (%f). New fields will be initialized to their defaulted value."), VersionNumber, SnapshotJsonVersion);
+		}
+	}
+
 	{
 		const TArray<TSharedPtr<FJsonValue>>& WindowsJsonArray = InRootJsonObject->GetArrayField(TEXT("Windows"));
 		for (const TSharedPtr<FJsonValue>& WindowJsonValue : WindowsJsonArray)
@@ -588,6 +620,22 @@ void FWidgetSnapshotData::LoadSnapshotFromJson(const TSharedRef<FJsonObject>& In
 		}
 	}
 
+	{
+		const TArray<TSharedPtr<FJsonValue>>& NavigationDataJsonArray = InRootJsonObject->GetArrayField(TEXT("NavigationData"));
+		for (const TSharedPtr<FJsonValue>& NavigationDataJsonValue : NavigationDataJsonArray)
+		{
+			FWidgetSnapshotNavigationSimulationData NavigationData;
+			const TSharedPtr<FJsonObject>& NavigationDataJsonObject = NavigationDataJsonValue->AsObject();
+			const TArray<TSharedPtr<FJsonValue>>& SimulationDataJsonArray = NavigationDataJsonObject->GetArrayField(TEXT("SimulationData"));
+			for (const TSharedPtr<FJsonValue>& SimulationDataJsonValue : SimulationDataJsonArray)
+			{
+				FNavigationSimulationWidgetNodePtr SimulationData = FNavigationSimulationWidgetNode::FromJson(SimulationDataJsonValue.ToSharedRef());
+				NavigationData.SimulationData.Add(SimulationData);
+			}
+			NavigationSimulationData.Add(MoveTemp(NavigationData));
+		}
+	}
+
 	{
 		const TArray<TSharedPtr<FJsonValue>>& TexturesJsonArray = InRootJsonObject->GetArrayField(TEXT("Textures"));
 		for (const TSharedPtr<FJsonValue>& TextureDataJsonValue : TexturesJsonArray)
