WidgetReflector: Fix some implicit conversions.

diff --git a/Engine/Source/Developer/SlateReflector/Private/InputEventVisualizer.cpp b/Engine/Source/Developer/SlateReflector/Private/InputEventVisualizer.cpp
--- a/Engine/Source/Developer/SlateReflector/Private/InputEventVisualizer.cpp
+++ b/Engine/Source/Developer/SlateReflector/Private/InputEventVisualizer.cpp
@@ -237,7 +237,7 @@ void FInputEventVisualizer::HandlePaint(const FPaintArgs& InArgs, const FGeometr
 			for (const FMouseEventInfo& Event : MouseEvents)
 			{
 				// Normalized animation value for the cursor ping between 0 and 1.
-				const float AnimAmount = (CurrentTime - Event.EventTime) / ClickFadeTime;
+				const float AnimAmount = (float)(CurrentTime - Event.EventTime) / ClickFadeTime;
 
 				const FVector2D CursorPosDesktopSpace = Event.CursorPingPosition;
 				const FVector2D PingSize = CursorSize * InputEventVisualizer::PingScaleAmount * FCurveHandle::ApplyEasing(AnimAmount, ECurveEaseFunction::QuadOut);
diff --git a/Engine/Source/Developer/SlateReflector/Private/Models/NavigationSimulationNode.cpp b/Engine/Source/Developer/SlateReflector/Private/Models/NavigationSimulationNode.cpp
--- a/Engine/Source/Developer/SlateReflector/Private/Models/NavigationSimulationNode.cpp
+++ b/Engine/Source/Developer/SlateReflector/Private/Models/NavigationSimulationNode.cpp
@@ -210,7 +210,7 @@ FNavigationSimulationWidgetNodePtr FNavigationSimulationWidgetNode::FromJson(con
 			const TArray<TSharedPtr<FJsonValue>>& StructJsonArray = InJsonValue->AsArray();
 			check(StructJsonArray.Num() == 2);
 
-			Out = FVector2D(StructJsonArray[0]->AsNumber(), StructJsonArray[1]->AsNumber());
+			Out = FVector2D((float)StructJsonArray[0]->AsNumber(), (float)StructJsonArray[1]->AsNumber());
 		}
 		static FNavigationSimulationWidgetInfo::TPointerAsInt ParsePtrIntFromString(const FString& Value)
 		{
@@ -228,10 +228,10 @@ FNavigationSimulationWidgetNodePtr FNavigationSimulationWidgetNode::FromJson(con
 			check(StructJsonArray.Num() == 4);
 
 			Out = FMatrix2x2(
-				StructJsonArray[0]->AsNumber(),
-				StructJsonArray[1]->AsNumber(),
-				StructJsonArray[2]->AsNumber(),
-				StructJsonArray[3]->AsNumber()
+				(float)StructJsonArray[0]->AsNumber(),
+				(float)StructJsonArray[1]->AsNumber(),
+				(float)StructJsonArray[2]->AsNumber(),
+				(float)StructJsonArray[3]->AsNumber()
 			);
 		}
 		static void ParseJsonValue(const TSharedPtr<FJsonValue>& InJsonValue, FSlateLayoutTransform& Out)
@@ -244,7 +244,7 @@ FNavigationSimulationWidgetNodePtr FNavigationSimulationWidgetNode::FromJson(con
 			check(StructJsonObject.IsValid());
 
 
-			float Scale = StructJsonObject->GetNumberField(TEXT("Scale"));
+			float Scale = (float)StructJsonObject->GetNumberField(TEXT("Scale"));
 			FVector2D Translation;
 			ParseJsonValue(StructJsonObject->GetField<EJson::None>(TEXT("Translation")), Translation);
 			Out = FSlateLayoutTransform(Scale, Translation);
diff --git a/Engine/Source/Developer/SlateReflector/Private/Widgets/SAtlasVisualizer.cpp b/Engine/Source/Developer/SlateReflector/Private/Widgets/SAtlasVisualizer.cpp
--- a/Engine/Source/Developer/SlateReflector/Private/Widgets/SAtlasVisualizer.cpp
+++ b/Engine/Source/Developer/SlateReflector/Private/Widgets/SAtlasVisualizer.cpp
@@ -491,7 +491,7 @@ FText SAtlasVisualizer::GetViewportSizeText() const
 FVector2D SAtlasVisualizer::GetViewportWidgetSize() const
 {
 	const FIntPoint ViewportSize = GetSize();
-	return FVector2D(ViewportSize.X, ViewportSize.Y);
+	return FVector2D((float)ViewportSize.X, (float)ViewportSize.Y);
 }
 
 FText SAtlasVisualizer::GetZoomLevelPercentText() const
diff --git a/Engine/Source/Developer/SlateReflector/Private/Widgets/SSlateOptions.cpp b/Engine/Source/Developer/SlateReflector/Private/Widgets/SSlateOptions.cpp
--- a/Engine/Source/Developer/SlateReflector/Private/Widgets/SSlateOptions.cpp
+++ b/Engine/Source/Developer/SlateReflector/Private/Widgets/SSlateOptions.cpp
@@ -93,8 +93,8 @@ void SSlateOptions::Construct( const FArguments& InArgs )
 			.AutoWidth()
 			[
 				SNew(SBox)
-				.MinDesiredWidth(100)
-				.MaxDesiredWidth(250)
+				.MinDesiredWidth(100.f)
+				.MaxDesiredWidth(250.f)
 				[
 					SNew(SSpinBox<float>)
 					.Value(this, &SSlateOptions::HandleAppScaleSliderValue)
diff --git a/Engine/Source/Developer/SlateReflector/Private/Widgets/SWidgetEventLog.cpp b/Engine/Source/Developer/SlateReflector/Private/Widgets/SWidgetEventLog.cpp
--- a/Engine/Source/Developer/SlateReflector/Private/Widgets/SWidgetEventLog.cpp
+++ b/Engine/Source/Developer/SlateReflector/Private/Widgets/SWidgetEventLog.cpp
@@ -124,7 +124,7 @@ void SWidgetEventLog::GenerateList()
 	auto GenerateCategory = [this] (EWidgetEventType EventType)
 	{
 		check(GetEnum(EventType)->GetMaxEnumValue() < (int64)TNumericLimits<uint8>::Max());
-		GetBitField(EventType) = TBitArray<>(false, GetEnum(EventType)->GetMaxEnumValue());
+		GetBitField(EventType) = TBitArray<>(false, (int32)GetEnum(EventType)->GetMaxEnumValue());
 		GetCategoryFlag(EventType) = false;
 	};
 	GenerateCategory(EWidgetEventType::Focus);
diff --git a/Engine/Source/Developer/SlateReflector/Private/Widgets/SWidgetHittestGrid.cpp b/Engine/Source/Developer/SlateReflector/Private/Widgets/SWidgetHittestGrid.cpp
--- a/Engine/Source/Developer/SlateReflector/Private/Widgets/SWidgetHittestGrid.cpp
+++ b/Engine/Source/Developer/SlateReflector/Private/Widgets/SWidgetHittestGrid.cpp
@@ -58,7 +58,7 @@ namespace WidgetHittestGridInternal
 		void Construct(const FArguments& InArgs, const TSharedRef<STableViewBase>& InOwnerTableView)
 		{
 			NodeToVisualize = InArgs._NodeToVisualize;
-			SMultiColumnTableRow<TSharedRef<FIntermediateResultNode>>::Construct(SMultiColumnTableRow<TSharedRef<FIntermediateResultNode>>::FArguments().Padding(0), InOwnerTableView);
+			SMultiColumnTableRow<TSharedRef<FIntermediateResultNode>>::Construct(SMultiColumnTableRow<TSharedRef<FIntermediateResultNode>>::FArguments().Padding(0.f), InOwnerTableView);
 		}
 
 	public:
