WidgetReflector: Add support for GlobalInvalidation.

diff --git a/Engine/Source/Developer/SlateReflector/Private/Models/WidgetReflectorNode.cpp b/Engine/Source/Developer/SlateReflector/Private/Models/WidgetReflectorNode.cpp
--- a/Engine/Source/Developer/SlateReflector/Private/Models/WidgetReflectorNode.cpp
+++ b/Engine/Source/Developer/SlateReflector/Private/Models/WidgetReflectorNode.cpp
@@ -158,6 +158,11 @@ bool FLiveWidgetReflectorNode::GetWidgetHasActiveTimers() const
 	return FWidgetReflectorNodeUtils::GetWidgetHasActiveTimers(Widget.Pin());
 }
 
+bool FLiveWidgetReflectorNode::GetWidgetIsInvalidationRoot() const
+{
+	return FWidgetReflectorNodeUtils::GetWidgetIsInvalidationRoot(Widget.Pin());
+}
+
 FText FLiveWidgetReflectorNode::GetWidgetReadableLocation() const
 {
 	return FWidgetReflectorNodeUtils::GetWidgetReadableLocation(Widget.Pin());
@@ -236,6 +241,7 @@ FSnapshotWidgetReflectorNode::FSnapshotWidgetReflectorNode(const FArrangedWidget
 	, bCachedWidgetIsVolatile(FWidgetReflectorNodeUtils::GetWidgetIsVolatile(InArrangedWidget.Widget))
 	, bCachedWidgetIsVolatileIndirectly(FWidgetReflectorNodeUtils::GetWidgetIsVolatileIndirectly(InArrangedWidget.Widget))
 	, bCachedWidgetHasActiveTimers(FWidgetReflectorNodeUtils::GetWidgetHasActiveTimers(InArrangedWidget.Widget))
+	, bCachedWidgetIsInvalidationRoot(FWidgetReflectorNodeUtils::GetWidgetIsInvalidationRoot(InArrangedWidget.Widget))
 	, bCachedWidgetEnabled(FWidgetReflectorNodeUtils::GetWidgetEnabled(InArrangedWidget.Widget))
 	, CachedWidgetClippingText(FWidgetReflectorNodeUtils::GetWidgetClippingText(InArrangedWidget.Widget))
 	, CachedWidgetReadableLocation(FWidgetReflectorNodeUtils::GetWidgetReadableLocation(InArrangedWidget.Widget))
@@ -308,6 +314,11 @@ bool FSnapshotWidgetReflectorNode::GetWidgetHasActiveTimers() const
 	return bCachedWidgetHasActiveTimers;
 }
 
+bool FSnapshotWidgetReflectorNode::GetWidgetIsInvalidationRoot() const
+{
+	return bCachedWidgetIsInvalidationRoot;
+}
+
 FText FSnapshotWidgetReflectorNode::GetWidgetReadableLocation() const
 {
 	return CachedWidgetReadableLocation;
@@ -445,6 +456,7 @@ TSharedRef<FJsonValue> FSnapshotWidgetReflectorNode::ToJson(const TSharedRef<FSn
 	RootJsonObject->SetBoolField(TEXT("WidgetIsVolatile"), RootSnapshotNode->bCachedWidgetIsVolatile);
 	RootJsonObject->SetBoolField(TEXT("WidgetIsVolatileIndirectly"), RootSnapshotNode->bCachedWidgetIsVolatileIndirectly);
 	RootJsonObject->SetBoolField(TEXT("WidgetHasActiveTimers"), RootSnapshotNode->bCachedWidgetHasActiveTimers);
+	RootJsonObject->SetBoolField(TEXT("WidgetIsInvalidationRoot"), RootSnapshotNode->bCachedWidgetIsInvalidationRoot);
 	RootJsonObject->SetBoolField(TEXT("WidgetEnabled"), RootSnapshotNode->bCachedWidgetEnabled);
 	RootJsonObject->SetStringField(TEXT("WidgetClippingText"), RootSnapshotNode->CachedWidgetClippingText.ToString());
 	RootJsonObject->SetStringField(TEXT("WidgetReadableLocation"), RootSnapshotNode->CachedWidgetReadableLocation.ToString());
@@ -480,8 +492,8 @@ TSharedRef<FSnapshotWidgetReflectorNode> FSnapshotWidgetReflectorNode::FromJson(
 			check(StructJsonArray.Num() == 2);
 
 			return FVector2D(
-				StructJsonArray[0]->AsNumber(), 
-				StructJsonArray[1]->AsNumber()
+				(float)StructJsonArray[0]->AsNumber(),
+				(float)StructJsonArray[1]->AsNumber()
 				);
 		}
 
@@ -495,10 +507,10 @@ TSharedRef<FSnapshotWidgetReflectorNode> FSnapshotWidgetReflectorNode::FromJson(
 			check(StructJsonArray.Num() == 4);
 
 			return FMatrix2x2(
-				StructJsonArray[0]->AsNumber(), 
-				StructJsonArray[1]->AsNumber(), 
-				StructJsonArray[2]->AsNumber(), 
-				StructJsonArray[3]->AsNumber()
+				(float)StructJsonArray[0]->AsNumber(), 
+				(float)StructJsonArray[1]->AsNumber(), 
+				(float)StructJsonArray[2]->AsNumber(),
+				(float)StructJsonArray[3]->AsNumber()
 				);
 		}
 
@@ -512,7 +524,7 @@ TSharedRef<FSnapshotWidgetReflectorNode> FSnapshotWidgetReflectorNode::FromJson(
 			check(StructJsonObject.IsValid());
 
 			return FSlateLayoutTransform(
-				StructJsonObject->GetNumberField(TEXT("Scale")),
+				(float)StructJsonObject->GetNumberField(TEXT("Scale")),
 				ParseVector2DJsonValue(StructJsonObject->GetField<EJson::None>(TEXT("Translation")))
 				);
 		}
@@ -542,10 +554,10 @@ TSharedRef<FSnapshotWidgetReflectorNode> FSnapshotWidgetReflectorNode::FromJson(
 			check(StructJsonArray.Num() == 4);
 
 			return FLinearColor(
-				StructJsonArray[0]->AsNumber(), 
-				StructJsonArray[1]->AsNumber(), 
-				StructJsonArray[2]->AsNumber(), 
-				StructJsonArray[3]->AsNumber()
+				(float)StructJsonArray[0]->AsNumber(),
+				(float)StructJsonArray[1]->AsNumber(),
+				(float)StructJsonArray[2]->AsNumber(),
+				(float)StructJsonArray[3]->AsNumber()
 				);
 		}
 
@@ -613,6 +625,7 @@ TSharedRef<FSnapshotWidgetReflectorNode> FSnapshotWidgetReflectorNode::FromJson(
 	RootSnapshotNode->bCachedWidgetIsVolatile = RootJsonObject->GetBoolField(TEXT("WidgetIsVolatile"));
 	RootSnapshotNode->bCachedWidgetIsVolatileIndirectly = RootJsonObject->GetBoolField(TEXT("WidgetIsVolatileIndirectly"));
 	RootSnapshotNode->bCachedWidgetHasActiveTimers = RootJsonObject->GetBoolField(TEXT("WidgetHasActiveTimers"));
+	RootSnapshotNode->bCachedWidgetIsInvalidationRoot = RootJsonObject->GetBoolField(TEXT("WidgetIsInvalidationRoot"));
 	RootSnapshotNode->bCachedWidgetEnabled = RootJsonObject->GetBoolField(TEXT("WidgetEnabled"));
 	RootSnapshotNode->CachedWidgetClippingText = FText::FromString(RootJsonObject->GetStringField(TEXT("WidgetClippingText")));
 	RootSnapshotNode->CachedWidgetReadableLocation = FText::FromString(RootJsonObject->GetStringField(TEXT("WidgetReadableLocation")));
@@ -867,6 +880,11 @@ bool FWidgetReflectorNodeUtils::GetWidgetHasActiveTimers(const TSharedPtr<const
 	return InWidget.IsValid() ? InWidget->HasActiveTimers() : false;
 }
 
+bool FWidgetReflectorNodeUtils::GetWidgetIsInvalidationRoot(const TSharedPtr<const SWidget>& InWidget)
+{
+	return InWidget.IsValid() ? InWidget->Advanced_IsInvalidationRoot() : false;
+}
+
 FText FWidgetReflectorNodeUtils::GetWidgetClippingText(const TSharedPtr<const SWidget>& InWidget)
 {
 	if ( InWidget.IsValid() )
diff --git a/Engine/Source/Developer/SlateReflector/Private/Models/WidgetReflectorNode.h b/Engine/Source/Developer/SlateReflector/Private/Models/WidgetReflectorNode.h
--- a/Engine/Source/Developer/SlateReflector/Private/Models/WidgetReflectorNode.h
+++ b/Engine/Source/Developer/SlateReflector/Private/Models/WidgetReflectorNode.h
@@ -77,18 +77,40 @@ public:
 	virtual FText GetWidgetClippingText() const = 0;
 
 	/**
-	* @return The bool indicating whether or not the widget we were initialized from reports as Focusable
-	*/
+	 * @return The bool indicating whether or not the widget we were initialized from reports as Focusable
+	 */
 	virtual bool GetWidgetFocusable() const = 0;
 
-	/** Is the widget visible? */
+	/**
+	 * @return The bool indicating whether or not the widget we were initialized from reports is visible
+	 */
 	virtual bool GetWidgetVisible() const = 0;
 
+	/**
+	 * @return The bool indicating whether or not the widget we were initialized from reports needs tick
+	 */
 	virtual bool GetWidgetNeedsTick() const = 0;
+
+	/**
+	 * @return The bool indicating whether or not the widget we were initialized from reports is volatile
+	 */
 	virtual bool GetWidgetIsVolatile() const = 0;
+
+	/**
+	 * @return The bool indicating whether or not the widget we were initialized from reports is volatile indirectly
+	 */
 	virtual bool GetWidgetIsVolatileIndirectly() const = 0;
+
+	/**
+	 * @return The bool indicating whether or not the widget we were initialized from reports has active timers
+	 */
 	virtual bool GetWidgetHasActiveTimers() const = 0;
 
+	/**
+	 * @return The bool indicating whether or not the widget we were initialized from reports is an Invalidation Root
+	 */
+	virtual bool GetWidgetIsInvalidationRoot() const = 0;
+
 	/**
 	 * The human readable location for widgets that are defined in C++ is the file and line number
 	 * The human readable location for widgets that are defined in UMG is the asset name
@@ -240,6 +262,7 @@ public:
 	virtual bool GetWidgetIsVolatile() const override;
 	virtual bool GetWidgetIsVolatileIndirectly() const override;
 	virtual bool GetWidgetHasActiveTimers() const override;
+	virtual bool GetWidgetIsInvalidationRoot() const override;
 	virtual FText GetWidgetReadableLocation() const override;
 	virtual FString GetWidgetFile() const override;
 	virtual int32 GetWidgetLineNumber() const override;
@@ -296,6 +319,7 @@ public:
 	virtual bool GetWidgetIsVolatile() const override;
 	virtual bool GetWidgetIsVolatileIndirectly() const override;
 	virtual bool GetWidgetHasActiveTimers() const override;
+	virtual bool GetWidgetIsInvalidationRoot() const override;
 	virtual FText GetWidgetReadableLocation() const override;
 	virtual FString GetWidgetFile() const override;
 	virtual int32 GetWidgetLineNumber() const override;
@@ -334,7 +358,7 @@ private:
 	/** The visibility string of the widget at the point it was passed to Initialize */
 	FText CachedWidgetVisibilityText;
 
-	/** Is the widget visible? */
+	/** The visible of the widget at the point it was passed to Initialize */
 	bool bCachedWidgetVisible;
 
 	/** The focusability of the widget at the point it was passed to Initialize */
@@ -351,6 +375,9 @@ private:
 
 	/** The active timer state of the widget at the point it was passed to Initialize */
 	bool bCachedWidgetHasActiveTimers;
+	
+	/** If the widget was an invalidation root at the point it was passed to Initialize */
+	bool bCachedWidgetIsInvalidationRoot;
 
 	/** The enabled state of the widget at the point it was passed to Initialize */
 	bool bCachedWidgetEnabled;
@@ -500,6 +527,7 @@ public:
 	static bool GetWidgetIsVolatile(const TSharedPtr<const SWidget>& InWidget);
 	static bool GetWidgetIsVolatileIndirectly(const TSharedPtr<const SWidget>& InWidget);
 	static bool GetWidgetHasActiveTimers(const TSharedPtr<const SWidget>& InWidget);
+	static bool GetWidgetIsInvalidationRoot(const TSharedPtr<const SWidget>& InWidget);
 	
 	/**
 	 * The human readable location for widgets that are defined in C++ is the file and line number
diff --git a/Engine/Source/Developer/SlateReflector/Private/VisualTreeCapture.cpp b/Engine/Source/Developer/SlateReflector/Private/VisualTreeCapture.cpp
--- a/Engine/Source/Developer/SlateReflector/Private/VisualTreeCapture.cpp
+++ b/Engine/Source/Developer/SlateReflector/Private/VisualTreeCapture.cpp
@@ -2,6 +2,7 @@
 
 #include "VisualTreeCapture.h"
 #include "Debugging/SlateDebugging.h"
+#include "FastUpdate/SlateInvalidationRoot.h"
 #include "Rendering/SlateRenderTransform.h"
 #include "Rendering/DrawElements.h"
 #include "Widgets/SWidget.h"
@@ -26,13 +27,37 @@ static bool IsPointInTriangle(const FVector2D& TestPoint, const FVector2D& A, co
 	return BA == CB && CB == AC;
 }
 
-FVisualEntry::FVisualEntry(int32 InElementIndex)
+FVisualEntry::FVisualEntry(const TWeakPtr<const SWidget>& InWidget, int32 InElementIndex)
 	: ElementIndex(InElementIndex)
+	, bFromCache(false)
+	, Widget(InWidget)
 {
 }
 
+FVisualEntry::FVisualEntry(const TSharedRef<const SWidget>& InWidget, const FSlateDrawElement& InElement)
+{
+	const FSlateRenderTransform& Transform = InElement.GetRenderTransform();
+	const FVector2D& LocalSize = InElement.GetLocalSize();
+
+	TopLeft = Transform.TransformPoint(FVector2D(0, 0));
+	TopRight = Transform.TransformPoint(FVector2D(LocalSize.X, 0));
+	BottomLeft = Transform.TransformPoint(FVector2D(0, LocalSize.Y));
+	BottomRight = Transform.TransformPoint(LocalSize);
+
+	LayerId = InElement.GetLayer();
+	ClippingIndex = INDEX_NONE;
+
+	bFromCache = true;
+	Widget = InWidget;
+}
+
 void FVisualEntry::Resolve(const FSlateWindowElementList& ElementList)
 {
+	if (bFromCache)
+	{
+		return;
+	}
+
 	const FSlateDrawElement& Element = ElementList.GetUncachedDrawElements()[ElementIndex];
 	const FSlateRenderTransform& Transform = Element.GetRenderTransform();
 	const FVector2D& LocalSize = Element.GetLocalSize();
@@ -63,10 +88,10 @@ TSharedPtr<const SWidget> FVisualTreeSnapshot::Pick(FVector2D Point)
 		const FVisualEntry& Entry = Entries[Index];
 		if (Entry.ClippingIndex != -1)
 		{
-			// TODO: This isn't always a valid index, but should be.
-			if (ensure(ClippingStates.IsValidIndex(Entry.ClippingIndex)))
+			const TArray<FSlateClippingState>& LocalClippingState = Entry.bFromCache ? CachedClippingStates : ClippingStates;
+			if (ensure(LocalClippingState.IsValidIndex(Entry.ClippingIndex)))
 			{
-				if (!ClippingStates[Entry.ClippingIndex].IsPointInside(Point))
+				if (!LocalClippingState[Entry.ClippingIndex].IsPointInside(Point))
 				{
 					continue;
 				}
@@ -78,15 +103,7 @@ TSharedPtr<const SWidget> FVisualTreeSnapshot::Pick(FVector2D Point)
 			continue;
 		}
 
-		TSharedPtr<const SWidget> Widget = Entry.Widget.Pin();
-
-		//if (Widget->GetVisibility() == EVisibility::HitTestInvisible)
-		//{
-		//	continue;
-		//}
-
-		// Build List.
-		return Widget;
+		return Entry.Widget.Pin();
 	}
 
 	return TSharedPtr<const SWidget>();
@@ -94,6 +111,9 @@ TSharedPtr<const SWidget> FVisualTreeSnapshot::Pick(FVector2D Point)
 
 FVisualTreeCapture::FVisualTreeCapture()
 	: bIsEnabled(false)
+	, WindowIsInvalidationRootCounter(0)
+	, WidgetIsInvalidationRootCounter(0)
+	, WidgetIsInvisibleToWidgetReflectorCounter(0)
 {
 }
 
@@ -114,6 +134,10 @@ void FVisualTreeCapture::Enable()
 		FSlateDebugging::EndWidgetPaint.AddRaw(this, &FVisualTreeCapture::EndWidgetPaint);
 		FSlateDebugging::ElementAdded.AddRaw(this, &FVisualTreeCapture::ElementAdded);
 		bIsEnabled = true;
+
+		WindowIsInvalidationRootCounter = 0;
+		WidgetIsInvalidationRootCounter = 0;
+		WidgetIsInvisibleToWidgetReflectorCounter = 0;
 	}
 #endif
 }
@@ -147,6 +171,37 @@ TSharedPtr<FVisualTreeSnapshot> FVisualTreeCapture::GetVisualTreeForWindow(SWind
 	return VisualTrees.FindRef(InWindow);
 }
 
+void FVisualTreeCapture::AddInvalidationRootCachedEntries(TSharedRef<FVisualTreeSnapshot> Tree, const FSlateInvalidationRoot* InvalidationRoot)
+{
+	check(InvalidationRoot);
+	const FSlateCachedElementData& Data = InvalidationRoot->GetCachedElements();
+	const TArray<TSharedPtr<FSlateCachedElementList>>& CachedElements = Data.GetCachedElementLists();
+	for (const TSharedPtr<FSlateCachedElementList>& CachedElement : CachedElements)
+	{
+		const FSlateDrawElementArray& DrawElementArray = CachedElement->DrawElements;
+		const SWidget* Widget = CachedElement->OwningWidget;
+		// todo, should check if parents has the metadata also
+		if (Widget && !Widget->GetMetaData<FInvisibleToWidgetReflectorMetaData>())
+		{
+			for (const FSlateDrawElement& Element : DrawElementArray)
+			{
+				const int32 EntryIndex = Tree->Entries.Emplace(Widget->AsShared(), Element);
+
+				const FSlateClippingState* ClippingState = Element.GetClippingHandle().GetCachedClipState();
+				if (ClippingState)
+				{
+					int32& ClippingRefIndex = Tree->Entries[EntryIndex].ClippingIndex;
+					ClippingRefIndex = Tree->CachedClippingStates.IndexOfByKey(*ClippingState);
+					if (ClippingRefIndex == INDEX_NONE)
+					{
+						ClippingRefIndex = Tree->CachedClippingStates.Add(*ClippingState);
+					}
+				}
+			}
+		}
+	}
+}
+
 void FVisualTreeCapture::BeginWindow(const FSlateWindowElementList& ElementList)
 {
 	TSharedPtr<FVisualTreeSnapshot> Tree = VisualTrees.FindRef(ElementList.GetPaintWindow());
@@ -158,10 +213,21 @@ void FVisualTreeCapture::BeginWindow(const FSlateWindowElementList& ElementList)
 
 	Tree->Entries.Reset();
 	Tree->ClippingStates.Reset();
+	Tree->CachedClippingStates.Reset();
+
+	if (ElementList.GetPaintWindow()->Advanced_IsInvalidationRoot())
+	{
+		++WindowIsInvalidationRootCounter;
+	}
 }
 
 void FVisualTreeCapture::EndWindow(const FSlateWindowElementList& ElementList)
 {
+	if (ElementList.GetPaintWindow()->Advanced_IsInvalidationRoot())
+	{
+		--WindowIsInvalidationRootCounter;
+	}
+
 	TSharedPtr<FVisualTreeSnapshot> Tree = VisualTrees.FindRef(ElementList.GetPaintWindow());
 	if (Tree.IsValid())
 	{
@@ -170,6 +236,13 @@ void FVisualTreeCapture::EndWindow(const FSlateWindowElementList& ElementList)
 			Entry.Resolve(ElementList);
 		}
 
+		if (ElementList.GetPaintWindow()->Advanced_IsInvalidationRoot())
+		{
+			// Add cached elements
+			const FSlateInvalidationRoot* InvalidationRoot = ElementList.GetPaintWindow()->Advanced_AsInvalidationRoot();
+			AddInvalidationRootCachedEntries(Tree.ToSharedRef(), InvalidationRoot);
+		}
+
 		Tree->ClippingStates = ElementList.GetClippingManager().GetClippingStates();
 		Tree->Entries.Sort([](const FVisualEntry& A, const FVisualEntry& B) {
 			return A.LayerId < B.LayerId;
@@ -183,6 +256,15 @@ void FVisualTreeCapture::BeginWidgetPaint(const SWidget* Widget, const FPaintArg
 	if (Tree.IsValid())
 	{
 		Tree->WidgetStack.Push(Widget->AsShared());
+
+		if (Widget->Advanced_IsInvalidationRoot())
+		{
+			++WidgetIsInvalidationRootCounter;
+		}
+		if (Widget->GetMetaData<FInvisibleToWidgetReflectorMetaData>())
+		{
+			++WidgetIsInvisibleToWidgetReflectorCounter;
+		}
 	}
 }
 
@@ -192,25 +274,36 @@ void FVisualTreeCapture::EndWidgetPaint(const SWidget* Widget, const FSlateWindo
 	if (Tree.IsValid())
 	{
 		Tree->WidgetStack.Pop();
+
+		if (Widget->Advanced_IsInvalidationRoot())
+		{
+			--WidgetIsInvalidationRootCounter;
+
+			// Add cached elements
+			const FSlateInvalidationRoot* InvalidationRoot = Widget->Advanced_AsInvalidationRoot();
+			AddInvalidationRootCachedEntries(Tree.ToSharedRef(), InvalidationRoot);
+		}
+		if (Widget->GetMetaData<FInvisibleToWidgetReflectorMetaData>())
+		{
+			--WidgetIsInvisibleToWidgetReflectorCounter;
+		}
 	}
 }
 
 void FVisualTreeCapture::ElementAdded(const FSlateWindowElementList& ElementList, int32 InElementIndex)
 {
+	if (WindowIsInvalidationRootCounter > 0 || WidgetIsInvalidationRootCounter > 0 || WidgetIsInvisibleToWidgetReflectorCounter > 0)
+	{
+		return;
+	}
+
 	TSharedPtr<FVisualTreeSnapshot> Tree = VisualTrees.FindRef(ElementList.GetPaintWindow());
 	if (Tree.IsValid())
 	{
 		if (Tree->WidgetStack.Num() > 0)
 		{
 			// Ignore any element added from a widget that's invisible to the widget reflector.
-			if (Tree->WidgetStack.Top().Pin()->GetMetaData<FInvisibleToWidgetReflectorMetaData>())
-			{
-				return;
-			}
-
-			FVisualEntry Entry(InElementIndex);
-			Entry.Widget = Tree->WidgetStack.Top();
-			Tree->Entries.Add(Entry);
+			Tree->Entries.Emplace(Tree->WidgetStack.Top(), InElementIndex);
 		}
 	}
 }
diff --git a/Engine/Source/Developer/SlateReflector/Private/VisualTreeCapture.h b/Engine/Source/Developer/SlateReflector/Private/VisualTreeCapture.h
--- a/Engine/Source/Developer/SlateReflector/Private/VisualTreeCapture.h
+++ b/Engine/Source/Developer/SlateReflector/Private/VisualTreeCapture.h
@@ -24,9 +24,11 @@ public:
 	int32 LayerId;
 	int32 ClippingIndex;
 	int32 ElementIndex;
+	bool bFromCache;
 	TWeakPtr<const SWidget> Widget;
 
-	FVisualEntry(int32 InElementIndex);
+	FVisualEntry(const TWeakPtr<const SWidget>& Widget, int32 InElementIndex);
+	FVisualEntry(const TSharedRef<const SWidget>& Widget, const FSlateDrawElement& InElement);
 
 	void Resolve(const FSlateWindowElementList& ElementList);
 
@@ -41,6 +43,7 @@ public:
 public:
 	TArray<FVisualEntry> Entries;
 	TArray<FSlateClippingState> ClippingStates;
+	TArray<FSlateClippingState> CachedClippingStates;
 	TArray<TWeakPtr<const SWidget>> WidgetStack;
 };
 
@@ -62,6 +65,10 @@ public:
 	TSharedPtr<FVisualTreeSnapshot> GetVisualTreeForWindow(SWindow* InWindow);
 	
 private:
+
+	void AddInvalidationRootCachedEntries(TSharedRef<FVisualTreeSnapshot> Tree, const FSlateInvalidationRoot* Entries);
+
+
 	void BeginWindow(const FSlateWindowElementList& ElementList);
 	void EndWindow(const FSlateWindowElementList& ElementList);
 
@@ -77,4 +84,7 @@ private:
 private:
 	TMap<const SWindow*, TSharedPtr<FVisualTreeSnapshot>> VisualTrees;
 	bool bIsEnabled;
+	int32 WindowIsInvalidationRootCounter;
+	int32 WidgetIsInvalidationRootCounter;
+	int32 WidgetIsInvisibleToWidgetReflectorCounter;
 };
\ No newline at end of file
diff --git a/Engine/Source/Developer/SlateReflector/Private/Widgets/SWidgetReflector.cpp b/Engine/Source/Developer/SlateReflector/Private/Widgets/SWidgetReflector.cpp
--- a/Engine/Source/Developer/SlateReflector/Private/Widgets/SWidgetReflector.cpp
+++ b/Engine/Source/Developer/SlateReflector/Private/Widgets/SWidgetReflector.cpp
@@ -286,23 +286,9 @@ private:
 
 	void SetPickingMode(EWidgetPickingMode InMode)
 	{
-#if WITH_SLATE_DEBUGGING
-		static auto CVarSlateGlobalInvalidation = IConsoleManager::Get().FindConsoleVariable(TEXT("Slate.EnableGlobalInvalidation"));
-#endif
-
 		if (PickingMode != InMode)
 		{
 			// Disable visual picking, and re-enable widget caching.
-#if WITH_SLATE_DEBUGGING
-			SInvalidationPanel::EnableInvalidationPanels(true);
-
-			if (PickingMode == EWidgetPickingMode::None)
-			{
-				bLastGlobalInvalidationState = CVarSlateGlobalInvalidation->GetBool();
-			}
-
-			CVarSlateGlobalInvalidation->Set(bLastGlobalInvalidationState);
-#endif
 			VisualCapture.Disable();
 
 			// Enable the picking mode.
@@ -312,19 +298,11 @@ private:
 			if (PickingMode == EWidgetPickingMode::HitTesting)
 			{
 				VisualCapture.Reset();
-#if WITH_SLATE_DEBUGGING
-				SInvalidationPanel::EnableInvalidationPanels(false);
-#endif
-				VisualCapture.Reset();
 			}
 			// If we're using the drawing picking mode enable it!
 			else if (PickingMode == EWidgetPickingMode::Drawable)
 			{
 				VisualCapture.Enable();
-#if WITH_SLATE_DEBUGGING
-				SInvalidationPanel::EnableInvalidationPanels(false);
-				CVarSlateGlobalInvalidation->Set(false);
-#endif
 			}
 		}
 	}
@@ -443,8 +421,6 @@ private:
 
 	FVisualTreeCapture VisualCapture;
 
-	bool bLastGlobalInvalidationState = false;
-
 private:
 	float SnapshotDelay;
 	bool bIsPendingDelayedSnapshot;
@@ -454,7 +430,7 @@ private:
 
 void SWidgetReflector::Construct( const FArguments& InArgs )
 {
-	// If saved, LoadSettings will override these varibles.
+	// If saved, LoadSettings will override these variables.
 	LastPickingMode = EWidgetPickingMode::HitTesting;
 	HiddenReflectorTreeColumns.Add(SReflectorTreeWidgetItem::NAME_Enabled.ToString());
 	HiddenReflectorTreeColumns.Add(SReflectorTreeWidgetItem::NAME_Volatile.ToString());
@@ -713,14 +689,14 @@ TSharedRef<SDockTab> SWidgetReflector::SpawnWidgetHierarchyTab(const FSpawnTabAr
 					[
 						SNew(SCheckBox)
 						.Style(FWidgetReflectorStyle::Get(), "CheckBoxNoHover")
-						.Padding(FMargin(4, 0))
+						.Padding(FMargin(4.f, 0.f))
 						.HAlign(HAlign_Left)
 						.IsChecked(this, &SWidgetReflector::HandleGetPickingButtonChecked)
 						.IsEnabled_Lambda([this]() { return !bIsPendingDelayedSnapshot; })
 						.OnCheckStateChanged(this, &SWidgetReflector::HandlePickingModeStateChanged)
 						[
 							SNew(SBox)
-							.MinDesiredWidth(175)
+							.MinDesiredWidth(175.f)
 							.VAlign(VAlign_Center)
 							[
 								SNew(SHorizontalBox)
@@ -829,7 +805,7 @@ TSharedRef<SDockTab> SWidgetReflector::SpawnWidgetHierarchyTab(const FSpawnTabAr
 			.FillHeight(1.0f)
 			[
 				SNew(SBorder)
-				.Padding(0)
+				.Padding(0.f)
 				.BorderImage(FCoreStyle::Get().GetBrush("ToolPanel.GroupBorder"))
 				[
 					// The tree view that shows all the info that we capture.
@@ -1573,8 +1549,8 @@ TSharedRef<SWidget> SWidgetReflector::HandleSnapshotOptionsTreeContextMenu()
 		.HAlign(HAlign_Right)
 		[
 			SNew(SSpinBox<float>)
-			.MinValue(0)
-			.MinDesiredWidth(40)
+			.MinValue(0.f)
+			.MinDesiredWidth(40.f)
 			.Value_Lambda([this]() { return SnapshotDelay; })
 			.OnValueCommitted_Lambda([this](const float InValue, ETextCommit::Type) { SnapshotDelay = FMath::Max(0.0f, InValue); })
 		];
diff --git a/Engine/Source/Developer/SlateReflector/Private/Widgets/SWidgetReflectorTreeWidgetItem.cpp b/Engine/Source/Developer/SlateReflector/Private/Widgets/SWidgetReflectorTreeWidgetItem.cpp
--- a/Engine/Source/Developer/SlateReflector/Private/Widgets/SWidgetReflectorTreeWidgetItem.cpp
+++ b/Engine/Source/Developer/SlateReflector/Private/Widgets/SWidgetReflectorTreeWidgetItem.cpp
@@ -29,11 +29,11 @@ void SReflectorTreeWidgetItem::Construct(const FArguments& InArgs, const TShared
 	this->WidgetInfo = InArgs._WidgetInfoToVisualize;
 	this->OnAccessSourceCode = InArgs._SourceCodeAccessor;
 	this->OnAccessAsset = InArgs._AssetAccessor;
-	this->SetPadding(0);
+	this->SetPadding(0.f);
 
 	check(WidgetInfo.IsValid());
 
-	SMultiColumnTableRow< TSharedRef<FWidgetReflectorNodeBase> >::Construct(SMultiColumnTableRow< TSharedRef<FWidgetReflectorNodeBase> >::FArguments().Padding(0), InOwnerTableView);
+	SMultiColumnTableRow< TSharedRef<FWidgetReflectorNodeBase> >::Construct(SMultiColumnTableRow< TSharedRef<FWidgetReflectorNodeBase> >::FArguments().Padding(0.f), InOwnerTableView);
 }
 
 BEGIN_SLATE_FUNCTION_BUILD_OPTIMIZATION
@@ -55,17 +55,29 @@ TSharedRef<SWidget> SReflectorTreeWidgetItem::GenerateWidgetForColumn(const FNam
 
 	if (ColumnName == NAME_WidgetName )
 	{
-		return SNew(SHorizontalBox)
+		TSharedRef<SHorizontalBox> HorizontalBox = SNew(SHorizontalBox);
 
-		+ SHorizontalBox::Slot()
+		HorizontalBox->AddSlot()
 		.AutoWidth()
 		[
 			SNew(SExpanderArrow, SharedThis(this))
 			.IndentAmount(16)
 			.ShouldDrawWires(true)
-		]
+		];
 
-		+ SHorizontalBox::Slot()
+		if (WidgetInfo->GetWidgetIsInvalidationRoot())
+		{
+			HorizontalBox->AddSlot()
+			.AutoWidth()
+			.Padding(2.0f, 0.0f)
+			.VAlign(VAlign_Center)
+			[
+				SNew(STextBlock)
+				.Text(LOCTEXT("InvalidationRoot_Short", "[IR]"))
+			];
+		}
+
+		HorizontalBox->AddSlot()
 		.AutoWidth()
 		.Padding(2.0f, 0.0f)
 		.VAlign(VAlign_Center)
@@ -74,6 +86,8 @@ TSharedRef<SWidget> SReflectorTreeWidgetItem::GenerateWidgetForColumn(const FNam
 			.Text(WidgetInfo->GetWidgetTypeAndShortName())
 			.ColorAndOpacity(this, &SReflectorTreeWidgetItem::GetTint)
 		];
+
+		return HorizontalBox;
 	}
 	else if (ColumnName == NAME_WidgetInfo )
 	{
@@ -169,7 +183,7 @@ TSharedRef<SWidget> SReflectorTreeWidgetItem::GenerateWidgetForColumn(const FNam
 
 			+ SHorizontalBox::Slot()
 			.AutoWidth()
-			.Padding(FMargin(0, 0, 2, 0))
+			.Padding(FMargin(0.f, 0.f, 2.f, 0.f))
 			[
 				SNew(SHyperlink)
 				.ToolTipText(LOCTEXT("ClickToCopy", "Click to copy address."))
diff --git a/Engine/Source/Runtime/Slate/Private/Widgets/SInvalidationPanel.cpp b/Engine/Source/Runtime/Slate/Private/Widgets/SInvalidationPanel.cpp
--- a/Engine/Source/Runtime/Slate/Private/Widgets/SInvalidationPanel.cpp
+++ b/Engine/Source/Runtime/Slate/Private/Widgets/SInvalidationPanel.cpp
@@ -270,6 +270,11 @@ bool SInvalidationPanel::Advanced_IsInvalidationRoot() const
 	return GetCanCache();
 }
 
+const FSlateInvalidationRoot* SInvalidationPanel::Advanced_AsInvalidationRoot() const
+{
+	return GetCanCache() ? this : nullptr;
+}
+
 int32 SInvalidationPanel::PaintSlowPath(const FSlateInvalidationContext& Context)
 {
 	return SCompoundWidget::OnPaint(*Context.PaintArgs, GetPaintSpaceGeometry(), Context.CullingRect, *Context.WindowElementList, Context.IncomingLayerId, Context.WidgetStyle, Context.bParentEnabled);
diff --git a/Engine/Source/Runtime/Slate/Public/Widgets/SInvalidationPanel.h b/Engine/Source/Runtime/Slate/Public/Widgets/SInvalidationPanel.h
--- a/Engine/Source/Runtime/Slate/Public/Widgets/SInvalidationPanel.h
+++ b/Engine/Source/Runtime/Slate/Public/Widgets/SInvalidationPanel.h
@@ -55,6 +55,7 @@ public:
 protected:
 	virtual bool CustomPrepass(float LayoutScaleMultiplier) override;
 	virtual bool Advanced_IsInvalidationRoot() const override;
+	virtual const FSlateInvalidationRoot* Advanced_AsInvalidationRoot() const override;
 	virtual int32 PaintSlowPath(const FSlateInvalidationContext& Context) override;
 private:
 	void OnGlobalInvalidationToggled(bool bGlobalInvalidationEnabled);
diff --git a/Engine/Source/Runtime/SlateCore/Public/FastUpdate/SlateInvalidationRoot.h b/Engine/Source/Runtime/SlateCore/Public/FastUpdate/SlateInvalidationRoot.h
--- a/Engine/Source/Runtime/SlateCore/Public/FastUpdate/SlateInvalidationRoot.h
+++ b/Engine/Source/Runtime/SlateCore/Public/FastUpdate/SlateInvalidationRoot.h
@@ -86,6 +86,7 @@ public:
 
 	/** @return the cached draw elements for this window and its widget hierarchy*/
 	FSlateCachedElementData& GetCachedElements() { return *CachedElementData; }
+	const FSlateCachedElementData& GetCachedElements() const { return *CachedElementData; }
 	const SWidget* GetInvalidationRootWidget() const;
 	int32 GetFastPathGenerationNumber() const { return FastPathGenerationNumber; }
 	FSlateInvalidationRootHandle GetInvalidationRootHandle() const { return InvalidationRootHandle; }
diff --git a/Engine/Source/Runtime/SlateCore/Public/Rendering/DrawElements.h b/Engine/Source/Runtime/SlateCore/Public/Rendering/DrawElements.h
--- a/Engine/Source/Runtime/SlateCore/Public/Rendering/DrawElements.h
+++ b/Engine/Source/Runtime/SlateCore/Public/Rendering/DrawElements.h
@@ -370,6 +370,7 @@ struct FSlateCachedElementList
 
 private:
 	SLATECORE_API void DestroyCachedData();
+
 public:
 	/** List of source draw elements to create batches from */
 	FSlateDrawElementArray DrawElements;
@@ -382,7 +383,6 @@ public:
 	FSlateCachedElementData* ParentData;
 
 	FSlateCachedFastPathRenderingData* CachedRenderingData;
-private:
 };
 
 struct FSlateCachedElementsHandle
@@ -428,6 +428,7 @@ struct FSlateCachedElementData
 	void CleanupUnusedClipStates();
 
 	const TSparseArray<FSlateRenderBatch>& GetCachedBatches() const { return CachedBatches; }
+	const TArray<TSharedPtr<FSlateCachedElementList>>& GetCachedElementLists() const { return CachedElementLists; }
 
 	void AddReferencedObjects(FReferenceCollector& Collector);
 
diff --git a/Engine/Source/Runtime/SlateCore/Public/Widgets/SWidget.h b/Engine/Source/Runtime/SlateCore/Public/Widgets/SWidget.h
--- a/Engine/Source/Runtime/SlateCore/Public/Widgets/SWidget.h
+++ b/Engine/Source/Runtime/SlateCore/Public/Widgets/SWidget.h
@@ -1373,6 +1373,7 @@ public:
 	/** Is this widget derivative of SWindow */
 	virtual bool Advanced_IsWindow() const { return false; }
 	virtual bool Advanced_IsInvalidationRoot() const { return false; }
+	virtual const FSlateInvalidationRoot* Advanced_AsInvalidationRoot() const { return nullptr; }
 protected:
 
 	/**
diff --git a/Engine/Source/Runtime/SlateCore/Public/Widgets/SWindow.h b/Engine/Source/Runtime/SlateCore/Public/Widgets/SWindow.h
--- a/Engine/Source/Runtime/SlateCore/Public/Widgets/SWindow.h
+++ b/Engine/Source/Runtime/SlateCore/Public/Widgets/SWindow.h
@@ -807,6 +807,7 @@ public:
 
 	virtual bool Advanced_IsWindow() const override { return true; }
 	virtual bool Advanced_IsInvalidationRoot() const override { return bAllowFastUpdate; }
+	virtual const FSlateInvalidationRoot* Advanced_AsInvalidationRoot() const override { return bAllowFastUpdate ? this : nullptr; }
 
 #if WITH_ACCESSIBILITY
 	virtual TSharedRef<FSlateAccessibleWidget> CreateAccessibleWidget() override;
diff --git a/Engine/Source/Runtime/UMG/Public/Slate/SRetainerWidget.h b/Engine/Source/Runtime/UMG/Public/Slate/SRetainerWidget.h
--- a/Engine/Source/Runtime/UMG/Public/Slate/SRetainerWidget.h
+++ b/Engine/Source/Runtime/UMG/Public/Slate/SRetainerWidget.h
@@ -93,6 +93,7 @@ protected:
 	virtual int32 OnPaint(const FPaintArgs& Args, const FGeometry& AllottedGeometry, const FSlateRect& MyCullingRect, FSlateWindowElementList& OutDrawElements, int32 LayerId, const FWidgetStyle& InWidgetStyle, bool bParentEnabled) const override;
 	virtual FVector2D ComputeDesiredSize(float Scale) const override;
 	virtual bool Advanced_IsInvalidationRoot() const { return bEnableRetainedRendering; }
+	virtual const FSlateInvalidationRoot* Advanced_AsInvalidationRoot() const override { return bEnableRetainedRendering ? this : nullptr; }
 	virtual bool CustomPrepass(float LayoutScaleMultiplier) override;
 
 	/** FSlateInvalidationRoot interface */
