WidgetReflector: Build widget tooltip only when requested. Add message when the widget is part of an invalidation root.

diff --git a/Engine/Source/Developer/SlateReflector/Private/Models/WidgetReflectorNode.cpp b/Engine/Source/Developer/SlateReflector/Private/Models/WidgetReflectorNode.cpp
--- a/Engine/Source/Developer/SlateReflector/Private/Models/WidgetReflectorNode.cpp
+++ b/Engine/Source/Developer/SlateReflector/Private/Models/WidgetReflectorNode.cpp
@@ -71,9 +71,10 @@ void FWidgetReflectorNodeBase::SetTint(const FLinearColor& InTint)
 	Tint = InTint;
 }
 
-void FWidgetReflectorNodeBase::AddChildNode(TSharedRef<FWidgetReflectorNodeBase> InChildNode)
+void FWidgetReflectorNodeBase::AddChildNode(TSharedRef<FWidgetReflectorNodeBase> InParentNode, TSharedRef<FWidgetReflectorNodeBase> InChildNode)
 {
-	ChildNodes.Add(MoveTemp(InChildNode));
+	InParentNode->ChildNodes.Add(MoveTemp(InChildNode));
+	InChildNode->ParentNode = InParentNode;
 }
 
 const TArray<TSharedRef<FWidgetReflectorNodeBase>>& FWidgetReflectorNodeBase::GetChildNodes() const
@@ -81,6 +82,10 @@ const TArray<TSharedRef<FWidgetReflectorNodeBase>>& FWidgetReflectorNodeBase::Ge
 	return ChildNodes;
 }
 
+const TSharedPtr<FWidgetReflectorNodeBase> FWidgetReflectorNodeBase::GetParentNode() const
+{
+	 return ParentNode.Pin();
+}
 
 /**
  * -----------------------------------------------------------------------------
@@ -642,7 +647,7 @@ TSharedRef<FSnapshotWidgetReflectorNode> FSnapshotWidgetReflectorNode::FromJson(
 	const TArray<TSharedPtr<FJsonValue>>& ChildNodesJsonArray = RootJsonObject->GetArrayField(TEXT("ChildNodes"));
 	for (const TSharedPtr<FJsonValue>& ChildNodeJsonValue : ChildNodesJsonArray)
 	{
-		RootSnapshotNode->AddChildNode(FSnapshotWidgetReflectorNode::FromJson(ChildNodeJsonValue.ToSharedRef()));
+		FSnapshotWidgetReflectorNode::AddChildNode(RootSnapshotNode, FSnapshotWidgetReflectorNode::FromJson(ChildNodeJsonValue.ToSharedRef()));
 	}
 
 	return RootSnapshotNode;
@@ -712,7 +717,7 @@ TSharedRef<FWidgetReflectorNodeBase> FWidgetReflectorNodeUtils::NewNodeTreeFrom(
 			}
 
 			// Note that we include both visible and invisible children!
-			NewNodeInstance->AddChildNode(NewNodeTreeFrom(InNodeType, FArrangedWidget(ChildWidget, ChildGeometry)));
+			FSnapshotWidgetReflectorNode::AddChildNode(NewNodeInstance, NewNodeTreeFrom(InNodeType, FArrangedWidget(ChildWidget, ChildGeometry)));
 		}
 	}
 
diff --git a/Engine/Source/Developer/SlateReflector/Private/Models/WidgetReflectorNode.h b/Engine/Source/Developer/SlateReflector/Private/Models/WidgetReflectorNode.h
--- a/Engine/Source/Developer/SlateReflector/Private/Models/WidgetReflectorNode.h
+++ b/Engine/Source/Developer/SlateReflector/Private/Models/WidgetReflectorNode.h
@@ -201,13 +201,18 @@ public:
 	/**
 	 * Add the given node to our list of children for this widget (this node will keep a strong reference to the instance)
 	 */
-	void AddChildNode(TSharedRef<FWidgetReflectorNodeBase> InChildNode);
+	static void AddChildNode(TSharedRef<FWidgetReflectorNodeBase> InParentNode, TSharedRef<FWidgetReflectorNodeBase> InChildNode);
 
 	/**
 	 * @return The node entries for the widget's children
 	 */
 	const TArray<TSharedRef<FWidgetReflectorNodeBase>>& GetChildNodes() const;
 
+	/**
+	 * @returns The node entry for the widget's parent, if it exists, 
+	 */
+	const TSharedPtr<FWidgetReflectorNodeBase> GetParentNode() const;
+
 protected:
 	/**
 	 * Default constructor
@@ -229,6 +234,9 @@ protected:
 	/** Node entries for the widget's children */
 	TArray<TSharedRef<FWidgetReflectorNodeBase>> ChildNodes;
 
+	/** Node entry for the widget's parent  */
+	TWeakPtr<FWidgetReflectorNodeBase> ParentNode;
+
 	/** A tint that is applied to text in order to provide visual hints (Transient) */
 	FLinearColor Tint;
 };
diff --git a/Engine/Source/Developer/SlateReflector/Private/Widgets/SWidgetReflector.cpp b/Engine/Source/Developer/SlateReflector/Private/Widgets/SWidgetReflector.cpp
--- a/Engine/Source/Developer/SlateReflector/Private/Widgets/SWidgetReflector.cpp
+++ b/Engine/Source/Developer/SlateReflector/Private/Widgets/SWidgetReflector.cpp
@@ -228,7 +228,7 @@ private:
 	 * @param InReflectorNode The node to generate the tool tip for.
 	 * @return The tool tip widget.
 	 */
-	TSharedRef<SToolTip> GenerateToolTipForReflectorNode( TSharedRef<FWidgetReflectorNodeBase> InReflectorNode );
+	TSharedPtr<IToolTip> GenerateToolTipForReflectorNode( TWeakPtr<FWidgetReflectorNodeBase> InReflectorNode ) const;
 
 	/**
 	 * Mark the provided reflector nodes such that they stand out in the tree and are visible.
@@ -1233,13 +1233,17 @@ void SWidgetReflector::SetSelectedAsReflectorTreeRoot()
 	}
 }
 
-TSharedRef<SToolTip> SWidgetReflector::GenerateToolTipForReflectorNode( TSharedRef<FWidgetReflectorNodeBase> InReflectorNode )
+TSharedPtr<IToolTip> SWidgetReflector::GenerateToolTipForReflectorNode( TWeakPtr<FWidgetReflectorNodeBase> InReflectorNode ) const
 {
-	return SNew(SToolTip)
-		[
-			SNew(SReflectorToolTipWidget)
-				.WidgetInfoToVisualize(InReflectorNode)
-		];
+	if (TSharedPtr<FWidgetReflectorNodeBase> ReflectorNode = InReflectorNode.Pin())
+	{
+		return SNew(SToolTip)
+			[
+				SNew(SReflectorToolTipWidget)
+				.WidgetInfoToVisualize(ReflectorNode)
+			];
+	}
+	return FSlateApplication::Get().MakeToolTip(LOCTEXT("MissingNode", "The node is invalid."));
 }
 
 void SWidgetReflector::VisualizeAsTree( const TArray<TSharedRef<FWidgetReflectorNodeBase>>& WidgetPathToVisualize )
diff --git a/Engine/Source/Developer/SlateReflector/Private/Widgets/SWidgetReflectorToolTipWidget.cpp b/Engine/Source/Developer/SlateReflector/Private/Widgets/SWidgetReflectorToolTipWidget.cpp
new file mode 100644
--- /dev/null
+++ b/Engine/Source/Developer/SlateReflector/Private/Widgets/SWidgetReflectorToolTipWidget.cpp
@@ -0,0 +1,104 @@
+// Copyright Epic Games, Inc. All Rights Reserved.
+
+#include "SWidgetReflectorToolTipWidget.h"
+
+#include "Layout/Geometry.h"
+#include "Styling/CoreStyle.h"
+#include "Widgets/SBoxPanel.h"
+#include "Widgets/Layout/SBorder.h"
+#include "Widgets/Layout/SGridPanel.h"
+#include "Widgets/Text/STextBlock.h"
+
+#define LOCTEXT_NAMESPACE "SReflectorToolTipWidget"
+
+const FText SReflectorToolTipWidget::TrueText = LOCTEXT("True", "True");
+const FText SReflectorToolTipWidget::FalseText = LOCTEXT("False", "False");
+
+void SReflectorToolTipWidget::Construct( const FArguments& InArgs )
+{
+	this->WidgetInfo = InArgs._WidgetInfoToVisualize;
+	check(WidgetInfo.IsValid());
+
+	const FGeometry Geom = FGeometry().MakeChild(WidgetInfo->GetLocalSize(), WidgetInfo->GetAccumulatedLayoutTransform(), WidgetInfo->GetAccumulatedRenderTransform(), FVector2D::ZeroVector);
+	SizeInfoText = FText::FromString(Geom.ToString());
+
+	TSharedRef<SGridPanel> GridPanel = SNew(SGridPanel)
+			.FillColumn(1, 1.0f);
+	{
+		int32 SlotCount = 0;
+		auto BuildLabelAndValue = [GridPanel, &SlotCount](const FText& Label, const TAttribute<FText>& Value)
+		{
+			GridPanel->AddSlot(0, SlotCount)
+			[
+				SNew(STextBlock)
+				.Text(Label)
+			];
+
+			GridPanel->AddSlot(1, SlotCount)
+			.Padding(5.0f, 0.0f, 0.0f, 0.0f)
+			[
+				SNew(STextBlock)
+				.Text(Value)
+			];
+
+			++SlotCount;
+		};
+
+		BuildLabelAndValue(LOCTEXT("DesiredSize", "Desired Size"), { this, &SReflectorToolTipWidget::GetWidgetsDesiredSize });
+		BuildLabelAndValue(LOCTEXT("ActualSize", "Actual Size"), { this, &SReflectorToolTipWidget::GetWidgetActualSize });
+		BuildLabelAndValue(LOCTEXT("SizeInfo", "Size Info"), { this, &SReflectorToolTipWidget::GetSizeInfo });
+		BuildLabelAndValue(LOCTEXT("Enabled", "Enabled"), { this, &SReflectorToolTipWidget::GetEnabled });
+		BuildLabelAndValue(LOCTEXT("NeedsTick", "Needs Tick"), { this, &SReflectorToolTipWidget::GetNeedsTick });
+		BuildLabelAndValue(LOCTEXT("IsVolatile", "Is Volatile"), { this, &SReflectorToolTipWidget::GetIsVolatile });
+		BuildLabelAndValue(LOCTEXT("IsVolatileIndirectly", "Is Volatile Indirectly"), { this, &SReflectorToolTipWidget::GetIsVolatileIndirectly });
+		BuildLabelAndValue(LOCTEXT("HasActiveTimers", "Has Active Timers"), { this, &SReflectorToolTipWidget::GetHasActiveTimers });
+	}
+
+	bool bIsInsideInvalidationRoot = false;
+	{
+		TSharedPtr<FWidgetReflectorNodeBase> CurrentWidget = WidgetInfo;
+		while(CurrentWidget)
+		{
+			if (CurrentWidget->GetWidgetIsInvalidationRoot())
+			{
+				bIsInsideInvalidationRoot = true;
+				break;
+			}
+			CurrentWidget = CurrentWidget->GetParentNode();
+		}
+	}
+
+	if (bIsInsideInvalidationRoot)
+	{
+		ChildSlot
+		[
+			SNew(SVerticalBox)
+			+ SVerticalBox::Slot()
+			.AutoHeight()
+			[
+				SNew(SBorder)
+				.BorderImage(FCoreStyle::Get().GetBrush("ToolPanel.GroupBorder"))
+				.HAlign(HAlign_Center)
+				[
+					SNew(STextBlock)
+					.Text(LOCTEXT("FromInvalidationRoot", "Controlled by an invalidation root"))
+				]
+			]
+			+ SVerticalBox::Slot()
+			.FillHeight(1.f)
+			[
+				GridPanel
+			]
+		];
+		
+	}
+	else
+	{
+		ChildSlot
+		[
+			GridPanel
+		];
+	}
+}
+
+#undef LOCTEXT_NAMESPACE
diff --git a/Engine/Source/Developer/SlateReflector/Private/Widgets/SWidgetReflectorToolTipWidget.h b/Engine/Source/Developer/SlateReflector/Private/Widgets/SWidgetReflectorToolTipWidget.h
--- a/Engine/Source/Developer/SlateReflector/Private/Widgets/SWidgetReflectorToolTipWidget.h
+++ b/Engine/Source/Developer/SlateReflector/Private/Widgets/SWidgetReflectorToolTipWidget.h
@@ -4,11 +4,8 @@
 
 #include "CoreMinimal.h"
 #include "Widgets/DeclarativeSyntaxSupport.h"
-#include "Layout/Geometry.h"
 #include "Widgets/SCompoundWidget.h"
-#include "Widgets/Text/STextBlock.h"
 #include "Models/WidgetReflectorNode.h"
-#include "Widgets/Layout/SGridPanel.h"
 
 #define LOCTEXT_NAMESPACE "SReflectorToolTipWidget"
 
@@ -28,140 +25,7 @@ public:
 public:
 
 	/** Constructs this widget with InArgs */
-	void Construct( const FArguments& InArgs )
-	{
-		this->WidgetInfo = InArgs._WidgetInfoToVisualize;
-		check(WidgetInfo.IsValid());
-
-		const FGeometry Geom = FGeometry().MakeChild(WidgetInfo->GetLocalSize(), WidgetInfo->GetAccumulatedLayoutTransform(), WidgetInfo->GetAccumulatedRenderTransform(), FVector2D::ZeroVector);
-		SizeInfoText = FText::FromString(Geom.ToString());
-
-		ChildSlot
-		[
-			SNew(SGridPanel)
-				.FillColumn(1, 1.0f)
-
-			// Desired Size
-			+ SGridPanel::Slot(0, 0)
-				[
-					SNew(STextBlock)
-						.Text(LOCTEXT("DesiredSize", "Desired Size"))
-				]
-
-			// Desired Size Value
-			+ SGridPanel::Slot(1, 0)
-				.Padding(5.0f, 0.0f, 0.0f, 0.0f)
-				[
-					SNew(STextBlock)
-						.Text(this, &SReflectorToolTipWidget::GetWidgetsDesiredSize)
-				]
-
-			// Actual Size
-			+ SGridPanel::Slot(0, 1)
-				[
-					SNew( STextBlock )
-						.Text(LOCTEXT("ActualSize", "Actual Size"))
-				]
-
-			// Actual Size Value
-			+ SGridPanel::Slot(1, 1)
-				.Padding(5.0f, 0.0f, 0.0f, 0.0f)
-				[
-					SNew(STextBlock)
-						.Text(this, &SReflectorToolTipWidget::GetWidgetActualSize)
-				]
-
-			// Size Info
-			+ SGridPanel::Slot(0, 2)
-			[
-				SNew(STextBlock)
-				.Text(LOCTEXT("SizeInfo", "Size Info"))
-			]
-
-			// Size Info Value
-			+ SGridPanel::Slot(1, 2)
-			.Padding(5.0f, 0.0f, 0.0f, 0.0f)
-			[
-				SNew(STextBlock)
-				.Text(this, &SReflectorToolTipWidget::GetSizeInfo)
-			]
-
-			// Enabled
-			+ SGridPanel::Slot(0, 3)
-			[
-				SNew(STextBlock)
-				.Text(LOCTEXT("Enabled", "Enabled"))
-			]
-
-			// Enabled Value
-			+ SGridPanel::Slot(1, 3)
-			.Padding(5.0f, 0.0f, 0.0f, 0.0f)
-			[
-				SNew(STextBlock)
-				.Text(this, &SReflectorToolTipWidget::GetEnabled)
-			]
-
-			// NeedsTick
-			+ SGridPanel::Slot(0, 4)
-			[
-				SNew(STextBlock)
-				.Text(LOCTEXT("NeedsTick", "Needs Tick"))
-			]
-
-			// NeedsTick Value
-			+ SGridPanel::Slot(1, 4)
-			.Padding(5.0f, 0.0f, 0.0f, 0.0f)
-			[
-				SNew(STextBlock)
-				.Text(this, &SReflectorToolTipWidget::GetNeedsTick)
-			]
-
-			// IsVolatile
-			+ SGridPanel::Slot(0, 5)
-			[
-				SNew(STextBlock)
-				.Text(LOCTEXT("IsVolatile", "Is Volatile"))
-			]
-
-			// IsVolatile Value
-			+ SGridPanel::Slot(1, 5)
-			.Padding(5.0f, 0.0f, 0.0f, 0.0f)
-			[
-				SNew(STextBlock)
-				.Text(this, &SReflectorToolTipWidget::GetIsVolatile)
-			]
-
-			// IsVolatileIndirectly
-			+ SGridPanel::Slot(0, 6)
-			[
-				SNew(STextBlock)
-				.Text(LOCTEXT("IsVolatileIndirectly", "Is Volatile Indirectly"))
-			]
-
-			// IsVolatileIndirectly Value
-			+ SGridPanel::Slot(1, 6)
-			.Padding(5.0f, 0.0f, 0.0f, 0.0f)
-			[
-				SNew(STextBlock)
-				.Text(this, &SReflectorToolTipWidget::GetIsVolatileIndirectly)
-			]
-
-			// HasActiveTimers
-			+ SGridPanel::Slot(0, 7)
-			[
-				SNew(STextBlock)
-				.Text(LOCTEXT("HasActiveTimers", "Has Active Timers"))
-			]
-
-			// HasActiveTimers Value
-			+ SGridPanel::Slot(1, 7)
-			.Padding(5.0f, 0.0f, 0.0f, 0.0f)
-			[
-				SNew(STextBlock)
-				.Text(this, &SReflectorToolTipWidget::GetHasActiveTimers)
-			]
-		];
-	}
+	void Construct(const FArguments& InArgs);
 
 private:
 
@@ -182,41 +46,34 @@ private:
 
 	FText GetEnabled() const
 	{
-		static const FText TrueText = LOCTEXT("True", "True");
-		static const FText FalseText = LOCTEXT("False", "False");
 		return (WidgetInfo->GetWidgetEnabled()) ? TrueText : FalseText;
 	}
 
 	FText GetNeedsTick() const
 	{
-		static const FText TrueText = LOCTEXT("True", "True");
-		static const FText FalseText = LOCTEXT("False", "False");
 		return (WidgetInfo->GetWidgetNeedsTick()) ? TrueText : FalseText;
 	}
 
 	FText GetIsVolatile() const
 	{
-		static const FText TrueText = LOCTEXT("True", "True");
-		static const FText FalseText = LOCTEXT("False", "False");
 		return (WidgetInfo->GetWidgetIsVolatile()) ? TrueText : FalseText;
 	}
 
 	FText GetIsVolatileIndirectly() const
 	{
-		static const FText TrueText = LOCTEXT("True", "True");
-		static const FText FalseText = LOCTEXT("False", "False");
 		return (WidgetInfo->GetWidgetIsVolatileIndirectly()) ? TrueText : FalseText;
 	}
 
 	FText GetHasActiveTimers() const
 	{
-		static const FText TrueText = LOCTEXT("True", "True");
-		static const FText FalseText = LOCTEXT("False", "False");
 		return (WidgetInfo->GetWidgetHasActiveTimers()) ? TrueText : FalseText;
 	}
 
 private:
 
+	static const FText TrueText;
+	static const FText FalseText;
+
 	/** The info about the widget that we are visualizing. */
 	TSharedPtr<FWidgetReflectorNodeBase> WidgetInfo;
 
