WidgetReflector: Set the correct fill value.
#rnx

diff --git a/Engine/Source/Developer/SlateReflector/Private/Widgets/SWidgetReflector.cpp b/Engine/Source/Developer/SlateReflector/Private/Widgets/SWidgetReflector.cpp
--- a/Engine/Source/Developer/SlateReflector/Private/Widgets/SWidgetReflector.cpp
+++ b/Engine/Source/Developer/SlateReflector/Private/Widgets/SWidgetReflector.cpp
@@ -826,7 +826,7 @@ TSharedRef<SDockTab> SWidgetReflector::SpawnWidgetHierarchyTab(const FSpawnTabAr
 
 						+SHeaderRow::Column(SReflectorTreeWidgetItem::NAME_WidgetName)
 						.DefaultLabel(LOCTEXT("WidgetName", "Widget Name"))
-						.FillWidth(0.80f)
+						.FillWidth(1.f)
 						.ShouldGenerateWidget(true)
 
 						+SHeaderRow::Column(SReflectorTreeWidgetItem::NAME_ForegroundColor)
