WidgetReflector: Show the number of registered attribute.

diff --git a/Engine/Source/Developer/SlateReflector/Private/Models/WidgetReflectorNode.cpp b/Engine/Source/Developer/SlateReflector/Private/Models/WidgetReflectorNode.cpp
--- a/Engine/Source/Developer/SlateReflector/Private/Models/WidgetReflectorNode.cpp
+++ b/Engine/Source/Developer/SlateReflector/Private/Models/WidgetReflectorNode.cpp
@@ -12,6 +12,7 @@
 #include "Layout/WidgetPath.h"
 #include "AssetRegistryModule.h"
 #include "Types/ReflectionMetadata.h"
+#include "Types/SlateAttributeMetaData.h"
 #include "FastUpdate/SlateInvalidationRoot.h"
 
 #define LOCTEXT_NAMESPACE "WidgetReflectorNode"
@@ -174,6 +175,16 @@ bool FLiveWidgetReflectorNode::GetWidgetIsInvalidationRoot() const
 	return FWidgetReflectorNodeUtils::GetWidgetIsInvalidationRoot(Widget.Pin());
 }
 
+int32 FLiveWidgetReflectorNode::GetWidgetAttributeCount() const
+{
+	return FWidgetReflectorNodeUtils::GetWidgetAttributeCount(Widget.Pin());
+}
+
+int32 FLiveWidgetReflectorNode::GetWidgetCollapsedAttributeCount() const
+{
+	return FWidgetReflectorNodeUtils::GetWidgetCollapsedAttributeCount(Widget.Pin());
+}
+
 FText FLiveWidgetReflectorNode::GetWidgetReadableLocation() const
 {
 	return FWidgetReflectorNodeUtils::GetWidgetReadableLocation(Widget.Pin());
@@ -259,6 +270,8 @@ FSnapshotWidgetReflectorNode::FSnapshotWidgetReflectorNode(const FArrangedWidget
 	, CachedWidgetReadableLocation(FWidgetReflectorNodeUtils::GetWidgetReadableLocation(InArrangedWidget.Widget))
 	, CachedWidgetFile(FWidgetReflectorNodeUtils::GetWidgetFile(InArrangedWidget.Widget))
 	, CachedWidgetLineNumber(FWidgetReflectorNodeUtils::GetWidgetLineNumber(InArrangedWidget.Widget))
+	, CachedWidgetAttributeCount(FWidgetReflectorNodeUtils::GetWidgetAttributeCount(InArrangedWidget.Widget))
+	, CachedWidgetCollapsedAttributeCount(FWidgetReflectorNodeUtils::GetWidgetCollapsedAttributeCount(InArrangedWidget.Widget))
 	, CachedWidgetAssetData(FWidgetReflectorNodeUtils::GetWidgetAssetData(InArrangedWidget.Widget))
 	, CachedWidgetDesiredSize(FWidgetReflectorNodeUtils::GetWidgetDesiredSize(InArrangedWidget.Widget))
 	, CachedWidgetForegroundColor(FWidgetReflectorNodeUtils::GetWidgetForegroundColor(InArrangedWidget.Widget))
@@ -351,6 +364,16 @@ int32 FSnapshotWidgetReflectorNode::GetWidgetLineNumber() const
 	return CachedWidgetLineNumber;
 }
 
+int32 FSnapshotWidgetReflectorNode::GetWidgetAttributeCount() const
+{
+	return CachedWidgetAttributeCount;
+}
+
+int32 FSnapshotWidgetReflectorNode::GetWidgetCollapsedAttributeCount() const
+{
+	return CachedWidgetCollapsedAttributeCount;
+}
+
 bool FSnapshotWidgetReflectorNode::HasValidWidgetAssetData() const
 {
 	return CachedWidgetAssetData.IsValid();
@@ -920,6 +943,30 @@ bool FWidgetReflectorNodeUtils::GetWidgetIsInvalidationRoot(const TSharedPtr<con
 	return InWidget.IsValid() ? InWidget->Advanced_IsInvalidationRoot() : false;
 }
 
+int32 FWidgetReflectorNodeUtils::GetWidgetAttributeCount(const TSharedPtr<const SWidget>& InWidget)
+{
+	if (InWidget.IsValid())
+	{
+		if (FSlateAttributeMetaData* MetaData = FSlateAttributeMetaData::FindMetaData(*InWidget.Get()))
+		{
+			return MetaData->RegisteredAttributeCount();
+		}
+	}
+	return 0;
+}
+
+int32 FWidgetReflectorNodeUtils::GetWidgetCollapsedAttributeCount(const TSharedPtr<const SWidget>& InWidget)
+{
+	if (InWidget.IsValid())
+	{
+		if (FSlateAttributeMetaData* MetaData = FSlateAttributeMetaData::FindMetaData(*InWidget.Get()))
+		{
+			return MetaData->RegisteredCollaspedAttributeCount();
+		}
+	}
+	return 0;
+}
+
 FText FWidgetReflectorNodeUtils::GetWidgetClippingText(const TSharedPtr<const SWidget>& InWidget)
 {
 	if ( InWidget.IsValid() )
diff --git a/Engine/Source/Developer/SlateReflector/Private/Models/WidgetReflectorNode.h b/Engine/Source/Developer/SlateReflector/Private/Models/WidgetReflectorNode.h
--- a/Engine/Source/Developer/SlateReflector/Private/Models/WidgetReflectorNode.h
+++ b/Engine/Source/Developer/SlateReflector/Private/Models/WidgetReflectorNode.h
@@ -116,6 +116,16 @@ public:
 	 */
 	virtual bool GetWidgetIsInvalidationRoot() const = 0;
 
+	/**
+	 * @return the numbers of registered slate attributes on the widget we were initialized from
+	 */
+	virtual int32 GetWidgetAttributeCount() const = 0;
+
+	/**
+	 * @return the numbers of registered slate attributes marked as "update when collapsed" on the widget we were initialized from
+	 */
+	virtual int32 GetWidgetCollapsedAttributeCount() const = 0;
+
 	/**
 	 * The human readable location for widgets that are defined in C++ is the file and line number
 	 * The human readable location for widgets that are defined in UMG is the asset name
@@ -280,6 +290,8 @@ public:
 	virtual FText GetWidgetReadableLocation() const override;
 	virtual FString GetWidgetFile() const override;
 	virtual int32 GetWidgetLineNumber() const override;
+	virtual int32 GetWidgetAttributeCount() const override;
+	virtual int32 GetWidgetCollapsedAttributeCount() const override;
 	virtual bool HasValidWidgetAssetData() const override;
 	virtual FAssetData GetWidgetAssetData() const override;
 	virtual FVector2D GetWidgetDesiredSize() const override;
@@ -338,6 +350,8 @@ public:
 	virtual FText GetWidgetReadableLocation() const override;
 	virtual FString GetWidgetFile() const override;
 	virtual int32 GetWidgetLineNumber() const override;
+	virtual int32 GetWidgetAttributeCount() const override;
+	virtual int32 GetWidgetCollapsedAttributeCount() const override;
 	virtual bool HasValidWidgetAssetData() const override;
 	virtual FAssetData GetWidgetAssetData() const override;
 	virtual FVector2D GetWidgetDesiredSize() const override;
@@ -412,6 +426,12 @@ private:
 	/** The line number of the file that the widget was created from at the point it was passed to Initialize (for C++ widgets) */
 	int32 CachedWidgetLineNumber;
 
+	/** The number of slate attributes registered on the SWidget. */
+	int32 CachedWidgetAttributeCount;
+	
+	/** The number of slate attributes registered on the SWidget marked as "update when collapsed". */
+	int32 CachedWidgetCollapsedAttributeCount;
+
 	/** The name of the asset that the widget was created from at the point it was passed to Initialize (for UMG widgets) */
 	FAssetData CachedWidgetAssetData;
 
@@ -549,6 +569,8 @@ public:
 	static bool GetWidgetIsVolatileIndirectly(const TSharedPtr<const SWidget>& InWidget);
 	static bool GetWidgetHasActiveTimers(const TSharedPtr<const SWidget>& InWidget);
 	static bool GetWidgetIsInvalidationRoot(const TSharedPtr<const SWidget>& InWidget);
+	static int32 GetWidgetAttributeCount(const TSharedPtr<const SWidget>& InWidget);
+	static int32 GetWidgetCollapsedAttributeCount(const TSharedPtr<const SWidget>& InWidget);
 	
 	/**
 	 * The human readable location for widgets that are defined in C++ is the file and line number
diff --git a/Engine/Source/Developer/SlateReflector/Private/Widgets/SWidgetReflectorToolTipWidget.cpp b/Engine/Source/Developer/SlateReflector/Private/Widgets/SWidgetReflectorToolTipWidget.cpp
--- a/Engine/Source/Developer/SlateReflector/Private/Widgets/SWidgetReflectorToolTipWidget.cpp
+++ b/Engine/Source/Developer/SlateReflector/Private/Widgets/SWidgetReflectorToolTipWidget.cpp
@@ -75,7 +75,9 @@ void SReflectorToolTipWidget::Construct( const FArguments& InArgs )
 		BuildLabelAndValue(LOCTEXT("NeedsTick", "Needs Tick"), { this, &SReflectorToolTipWidget::GetNeedsTick });
 		BuildLabelAndValue(LOCTEXT("IsVolatile", "Is Volatile"), { this, &SReflectorToolTipWidget::GetIsVolatile });
 		BuildLabelAndValue(LOCTEXT("IsVolatileIndirectly", "Is Volatile Indirectly"), { this, &SReflectorToolTipWidget::GetIsVolatileIndirectly });
-		BuildLabelAndValue(LOCTEXT("HasActiveTimers", "Has Active Timers"), { this, &SReflectorToolTipWidget::GetHasActiveTimers });
+		BuildLabelAndValue(LOCTEXT("HasActiveTimers", "Has Active Timer"), { this, &SReflectorToolTipWidget::GetHasActiveTimers });
+		BuildLabelAndValue(LOCTEXT("HasRegisteredAttribute", "Attribute Count"), { this, &SReflectorToolTipWidget::GetAttributeCount });
+		BuildLabelAndValue(LOCTEXT("HasCollapsedAttribute", "Collapsed Attribute Count"), { this, &SReflectorToolTipWidget::GetCollapsedAttributeCount });
 
 		if (bIsInsideInvalidationRoot)
 		{
diff --git a/Engine/Source/Developer/SlateReflector/Private/Widgets/SWidgetReflectorToolTipWidget.h b/Engine/Source/Developer/SlateReflector/Private/Widgets/SWidgetReflectorToolTipWidget.h
--- a/Engine/Source/Developer/SlateReflector/Private/Widgets/SWidgetReflectorToolTipWidget.h
+++ b/Engine/Source/Developer/SlateReflector/Private/Widgets/SWidgetReflectorToolTipWidget.h
@@ -69,6 +69,16 @@ private:
 		return (WidgetInfo->GetWidgetHasActiveTimers()) ? TrueText : FalseText;
 	}
 
+	FText GetAttributeCount() const
+	{
+		return FText::AsNumber(WidgetInfo->GetWidgetAttributeCount());
+	}
+
+	FText GetCollapsedAttributeCount() const
+	{
+		return FText::AsNumber(WidgetInfo->GetWidgetCollapsedAttributeCount());
+	}
+
 	FText GetIsVisible() const
 	{
 		return (WidgetInfo->GetWidgetVisible()) ? TrueText : FalseText;
