WidgetReflector: Save and load the attribute count in the snapshot.

diff --git a/Engine/Source/Developer/SlateReflector/Private/Models/WidgetReflectorNode.cpp b/Engine/Source/Developer/SlateReflector/Private/Models/WidgetReflectorNode.cpp
--- a/Engine/Source/Developer/SlateReflector/Private/Models/WidgetReflectorNode.cpp
+++ b/Engine/Source/Developer/SlateReflector/Private/Models/WidgetReflectorNode.cpp
@@ -249,6 +249,8 @@ TSharedRef<FSnapshotWidgetReflectorNode> FSnapshotWidgetReflectorNode::Create(co
 FSnapshotWidgetReflectorNode::FSnapshotWidgetReflectorNode()
 	: bCachedWidgetEnabled(false)
 	, CachedWidgetLineNumber(0)
+	, CachedWidgetAttributeCount(0)
+	, CachedWidgetCollapsedAttributeCount(0)
 {
 }
 
@@ -507,6 +509,8 @@ TSharedRef<FJsonValue> FSnapshotWidgetReflectorNode::ToJson(const TSharedRef<FSn
 	RootJsonObject->SetStringField(TEXT("WidgetReadableLocation"), RootSnapshotNode->CachedWidgetReadableLocation.ToString());
 	RootJsonObject->SetStringField(TEXT("WidgetFile"), RootSnapshotNode->CachedWidgetFile);
 	RootJsonObject->SetNumberField(TEXT("WidgetLineNumber"), RootSnapshotNode->CachedWidgetLineNumber);
+	RootJsonObject->SetNumberField(TEXT("WidgetAttributeCount"), RootSnapshotNode->CachedWidgetAttributeCount);
+	RootJsonObject->SetNumberField(TEXT("WidgetCollapsedAttributeCount"), RootSnapshotNode->CachedWidgetCollapsedAttributeCount);
 	RootJsonObject->SetField(TEXT("WidgetDesiredSize"), Internal::CreateVector2DJsonValue(RootSnapshotNode->CachedWidgetDesiredSize));
 	RootJsonObject->SetField(TEXT("WidgetForegroundColor"), Internal::CreateSlateColorJsonValue(RootSnapshotNode->CachedWidgetForegroundColor));
 	RootJsonObject->SetStringField(TEXT("WidgetAddress"), Internal::ConvertPtrIntToString(RootSnapshotNode->CachedWidgetAddress));
@@ -677,6 +681,8 @@ TSharedRef<FSnapshotWidgetReflectorNode> FSnapshotWidgetReflectorNode::FromJson(
 	RootSnapshotNode->CachedWidgetReadableLocation = FText::FromString(RootJsonObject->GetStringField(TEXT("WidgetReadableLocation")));
 	RootSnapshotNode->CachedWidgetFile = RootJsonObject->GetStringField(TEXT("WidgetFile"));
 	RootSnapshotNode->CachedWidgetLineNumber = RootJsonObject->GetIntegerField(TEXT("WidgetLineNumber"));
+	RootSnapshotNode->CachedWidgetAttributeCount = RootJsonObject->GetIntegerField(TEXT("WidgetAttributeCount"));
+	RootSnapshotNode->CachedWidgetCollapsedAttributeCount = RootJsonObject->GetIntegerField(TEXT("WidgetCollapsedAttributeCount"));
 	RootSnapshotNode->CachedWidgetDesiredSize = Internal::ParseVector2DJsonValue(RootJsonObject->GetField<EJson::None>(TEXT("WidgetDesiredSize")));
 	RootSnapshotNode->CachedWidgetForegroundColor = Internal::ParseSlateColorJsonValue(RootJsonObject->GetField<EJson::None>(TEXT("WidgetForegroundColor")));
 	RootSnapshotNode->CachedWidgetAddress = Internal::ParsePtrIntFromString(RootJsonObject->GetStringField(TEXT("WidgetAddress")));
diff --git a/Engine/Source/Developer/SlateReflector/Private/Widgets/SWidgetSnapshotVisualizer.cpp b/Engine/Source/Developer/SlateReflector/Private/Widgets/SWidgetSnapshotVisualizer.cpp
--- a/Engine/Source/Developer/SlateReflector/Private/Widgets/SWidgetSnapshotVisualizer.cpp
+++ b/Engine/Source/Developer/SlateReflector/Private/Widgets/SWidgetSnapshotVisualizer.cpp
@@ -446,7 +446,7 @@ void FWidgetSnapshotData::SaveSnapshotToBuffer(TArray<uint8>& OutData) const
 	BufferWriter.SerializeCompressed(TmpJsonData.GetData(), TmpJsonData.Num(), NAME_Zlib);
 }
 
-double SnapshotJsonVersion = 2.0;
+double SnapshotJsonVersion = 2.1;
 
 TSharedRef<FJsonObject> FWidgetSnapshotData::SaveSnapshotAsJson() const
 {
