WidgetReflector: Group the Visible flag with the Volatile flag in the widget tooltip.
#rnx

diff --git a/Engine/Source/Developer/SlateReflector/Private/Widgets/SWidgetReflectorToolTipWidget.cpp b/Engine/Source/Developer/SlateReflector/Private/Widgets/SWidgetReflectorToolTipWidget.cpp
--- a/Engine/Source/Developer/SlateReflector/Private/Widgets/SWidgetReflectorToolTipWidget.cpp
+++ b/Engine/Source/Developer/SlateReflector/Private/Widgets/SWidgetReflectorToolTipWidget.cpp
@@ -76,6 +76,11 @@ void SReflectorToolTipWidget::Construct( const FArguments& InArgs )
 		BuildLabelAndValue(LOCTEXT("NeedsTick", "Needs Tick"), { this, &SReflectorToolTipWidget::GetNeedsTick });
 		BuildLabelAndValue(LOCTEXT("IsVolatile", "Is Volatile"), { this, &SReflectorToolTipWidget::GetIsVolatile });
 		BuildLabelAndValue(LOCTEXT("IsVolatileIndirectly", "Is Volatile Indirectly"), { this, &SReflectorToolTipWidget::GetIsVolatileIndirectly });
+		if (bIsInsideInvalidationRoot)
+		{
+			BuildLabelAndValue(LOCTEXT("IsVisible", "Is Visible"), { this, &SReflectorToolTipWidget::GetIsVisible });
+			BuildLabelAndValue(LOCTEXT("IsVisibleInherited", "Is Visible Inherited"), { this, &SReflectorToolTipWidget::GetIsVisibleInherited });
+		}
 		BuildLabelAndValue(LOCTEXT("HasActiveTimers", "Has Active Timer"), { this, &SReflectorToolTipWidget::GetHasActiveTimers });
 		if (WidgetInfo->GetNodeType() == EWidgetReflectorNodeType::Live)
 		{
@@ -86,12 +91,6 @@ void SReflectorToolTipWidget::Construct( const FArguments& InArgs )
 			BuildLabelAndValue(LOCTEXT("HasRegisteredAttribute", "Attributes"), { this, &SReflectorToolTipWidget::GetAttributeCount });
 			BuildLabelAndValue(LOCTEXT("HasCollapsedAttribute", "Affect Visibility Attributes"), { this, &SReflectorToolTipWidget::GetCollapsedAttributeCount });
 		}
-
-		if (bIsInsideInvalidationRoot)
-		{
-			BuildLabelAndValue(LOCTEXT("IsVisible", "Is Visible"), { this, &SReflectorToolTipWidget::GetIsVisible });
-			BuildLabelAndValue(LOCTEXT("IsVisibleInherited", "Is Visible Inherited"), { this, &SReflectorToolTipWidget::GetIsVisibleInherited });
-		}
 	}
 
 	if (bIsInsideInvalidationRoot)
@@ -147,7 +146,7 @@ FText SReflectorToolTipWidget::GetLiveAttributeName() const
 		}
 	}
 
-	return FText::AsNumber(0);
+	return FText::GetEmpty();
 }
 
 #undef LOCTEXT_NAMESPACE
