- Upgraded Widget Reflector UI to UE5

diff --git a/Engine/Source/Developer/SlateReflector/Private/Styling/WidgetReflectorStyle.cpp b/Engine/Source/Developer/SlateReflector/Private/Styling/WidgetReflectorStyle.cpp
--- a/Engine/Source/Developer/SlateReflector/Private/Styling/WidgetReflectorStyle.cpp
+++ b/Engine/Source/Developer/SlateReflector/Private/Styling/WidgetReflectorStyle.cpp
@@ -5,9 +5,11 @@
 #include "Styling/SlateTypes.h"
 #include "Styling/CoreStyle.h"
 #include "Styling/SlateStyleMacros.h"
+#include "Styling/ToolBarStyle.h"
 
 // This is to fix the issue that SlateStyleMacros like IMAGE_BRUSH look for RootToContentDir but StyleSet->RootToContentDir is how this style is set up
 #define RootToContentDir StyleSet->RootToContentDir
+#define RootToCoreContentDir StyleSet->RootToCoreContentDir
 
 TSharedPtr< FSlateStyleSet > FWidgetReflectorStyle::StyleInstance = nullptr;
 
@@ -43,50 +45,29 @@ TSharedRef< FSlateStyleSet > FWidgetReflectorStyle::Create()
 	StyleSet->SetContentRoot(FPaths::EngineContentDir() / TEXT("Editor/Slate"));
 	StyleSet->SetCoreContentRoot(FPaths::EngineContentDir() / TEXT("Slate"));
 
-#if WITH_EDITOR
-	{
-		FButtonStyle Button = FButtonStyle()
-			.SetNormal(FSlateBoxBrush(RootToContentDir("Common/ButtonHoverHint.png"), FMargin(4 / 16.0f), FLinearColor(1, 1, 1, 0.15f)))
-			.SetHovered(FSlateBoxBrush(RootToContentDir("Common/ButtonHoverHint.png"), FMargin(4 / 16.0f), FLinearColor(1, 1, 1, 0.25f)))
-			.SetPressed(FSlateBoxBrush(RootToContentDir("Common/ButtonHoverHint.png"), FMargin(4 / 16.0f), FLinearColor(1, 1, 1, 0.30f)))
-			.SetNormalPadding(FMargin(0, 0, 0, 1))
-			.SetPressedPadding(FMargin(0, 1, 0, 0));
-		StyleSet->Set("Button", Button);
-
-		FCheckBoxStyle CustomCheckBoxStyle = FCoreStyle::Get().GetWidgetStyle<FCheckBoxStyle>("ToggleButtonCheckbox");
-		CustomCheckBoxStyle.SetUncheckedImage(CustomCheckBoxStyle.UncheckedHoveredImage);
-		CustomCheckBoxStyle.UncheckedImage.TintColor = FLinearColor(1, 1, 1, 0.1f);	
-		StyleSet->Set("CheckBox", CustomCheckBoxStyle);
-
-		CustomCheckBoxStyle.SetUncheckedHoveredImage(Button.Hovered);
-		StyleSet->Set("CheckBoxNoHover", CustomCheckBoxStyle);
-
-		FComboButtonStyle ComboButton = FComboButtonStyle()
-			.SetButtonStyle(Button.SetNormal(FSlateNoResource()))
-			.SetDownArrowImage(FSlateImageBrush(StyleSet->RootToCoreContentDir(TEXT("Common/ComboArrow.png")), Icon8x8))
-			.SetMenuBorderBrush(FSlateBoxBrush(StyleSet->RootToCoreContentDir(TEXT("Old/Menu_Background.png")), FMargin(8.0f / 64.0f)))
-			.SetMenuBorderPadding(FMargin(0.0f));
-		StyleSet->Set("ComboButton", ComboButton);
-
-		StyleSet->Set("SplitterDark", FSplitterStyle()
-			.SetHandleNormalBrush(FSlateColorBrush(FLinearColor(FColor(32, 32, 32))))
-			.SetHandleHighlightBrush(FSlateColorBrush(FLinearColor(FColor(96, 96, 96))))
-		);
-	}
-	
 	{
 		StyleSet->Set("Icon.FocusPicking", new FSlateImageBrush(RootToContentDir("Icons/SlateReflector/FocusPicking_24x.png"), Icon24x24));
-		StyleSet->Set("Icon.HitTestPicking", new FSlateImageBrush(RootToContentDir("Icons/SlateReflector/HitTestPicking_24x.png"), Icon24x24));
-		StyleSet->Set("Icon.VisualPicking", new FSlateImageBrush(RootToContentDir("Icons/SlateReflector/VisualPicking_24x.png"), Icon24x24));
+		StyleSet->Set("Icon.HitTestPicking", new FSlateImageBrush(RootToContentDir("Icons/GeneralTools/Select_40x.png"), Icon24x24));
+		StyleSet->Set("Icon.VisualPicking", new FSlateImageBrush(RootToContentDir("Icons/GeneralTools/Paint_40x.png"), Icon24x24));
+		StyleSet->Set("Icon.LoadSnapshot", new FSlateImageBrush(RootToContentDir("Icons/GeneralTools/Import_40x.png"), Icon24x24));
+		StyleSet->Set("Icon.Filter", new FSlateImageBrush(RootToContentDir("Icons/GeneralTools/Filter_40x.png"), Icon24x24));
+		StyleSet->Set("Icon.TakeSnapshot", new IMAGE_BRUSH_SVG("Starship/Common/SaveThumbnail", Icon24x24));
 
-		StyleSet->Set("Symbols.LeftArrow", new FSlateImageBrush(RootToContentDir("Common/LeftArrow.png"), Icon24x24));
-		StyleSet->Set("Symbols.RightArrow", new FSlateImageBrush(RootToContentDir("Common/SubmenuArrow.png"), Icon24x24));
-		StyleSet->Set("Symbols.UpArrow", new FSlateImageBrush(RootToContentDir("Common/UpArrow.png"), Icon24x24));
-		StyleSet->Set("Symbols.DownArrow", new FSlateImageBrush(RootToContentDir("Common/DownArrow.png"), Icon24x24));
+		StyleSet->Set("Symbols.LeftArrow", new CORE_IMAGE_BRUSH_SVG("Starship/Common/arrow-left", Icon24x24));
+		StyleSet->Set("Symbols.RightArrow", new CORE_IMAGE_BRUSH_SVG("Starship/Common/arrow-right", Icon24x24));
+		StyleSet->Set("Symbols.UpArrow", new CORE_IMAGE_BRUSH_SVG("Starship/Common/arrow-up", Icon24x24));
+		StyleSet->Set("Symbols.DownArrow", new CORE_IMAGE_BRUSH_SVG("Starship/Common/arrow-down", Icon24x24));
 
 		StyleSet->Set("WidgetReflector.TabIcon", new IMAGE_BRUSH_SVG("Starship/Common/Widget", Icon16x16));
+		StyleSet->Set("Icon.Ellipsis", new CORE_IMAGE_BRUSH_SVG("Starship/Common/ellipsis-vertical-narrow", FVector2D(6, 24)));
+	}
+
+	{
+		FToolBarStyle SlimToolbarStyle = FAppStyle::Get().GetWidgetStyle<FToolBarStyle>("SlimToolBar");
+		const FTextBlockStyle& ButtonText = FAppStyle::Get().GetWidgetStyle<FTextBlockStyle>("ButtonText");
+		SlimToolbarStyle.SetLabelStyle(FTextBlockStyle(ButtonText));
+		StyleSet->Set("BoldSlimToolbar", SlimToolbarStyle);
 	}
-#endif
 
 	return StyleSet;
 }
@@ -95,3 +76,6 @@ const ISlateStyle& FWidgetReflectorStyle::Get()
 {
 	return *StyleInstance;
 }
+
+#undef RootToCoreContentDir
+#undef RootToContentDir
\ No newline at end of file
diff --git a/Engine/Source/Developer/SlateReflector/Private/Widgets/SSlateOptions.cpp b/Engine/Source/Developer/SlateReflector/Private/Widgets/SSlateOptions.cpp
--- a/Engine/Source/Developer/SlateReflector/Private/Widgets/SSlateOptions.cpp
+++ b/Engine/Source/Developer/SlateReflector/Private/Widgets/SSlateOptions.cpp
@@ -70,8 +70,9 @@ void SSlateOptions::Construct( const FArguments& InArgs )
 	};
 
 	FToolBarBuilder ToolbarBuilder(TSharedPtr<const FUICommandList>(), FMultiBoxCustomization::None);
+	ToolbarBuilder.SetStyle(&FWidgetReflectorStyle::Get(), "BoldSlimToolbar");
 	Local::FillToolbar(ToolbarBuilder);
-
+	
 	ChildSlot
 	[
 		SNew(SBorder)
diff --git a/Engine/Source/Developer/SlateReflector/Private/Widgets/SWidgetReflector.cpp b/Engine/Source/Developer/SlateReflector/Private/Widgets/SWidgetReflector.cpp
--- a/Engine/Source/Developer/SlateReflector/Private/Widgets/SWidgetReflector.cpp
+++ b/Engine/Source/Developer/SlateReflector/Private/Widgets/SWidgetReflector.cpp
@@ -123,6 +123,10 @@ namespace WidgetReflectorIcon
 	static const FName FocusPicking = "Icon.FocusPicking";
 	static const FName HitTestPicking = "Icon.HitTestPicking";
 	static const FName VisualPicking = "Icon.VisualPicking";
+	static const FName Ellipsis = "Icon.Ellipsis";
+	static const FName Filter = "Icon.Filter";
+	static const FName LoadSnapshot = "Icon.LoadSnapshot";
+	static const FName TakeSnapshot = "Icon.TakeSnapshot";
 }
 
 enum class EWidgetPickingMode : uint8
@@ -278,8 +282,8 @@ private:
 
 	//~ Handle for the picking button
 	ECheckBoxState HandleGetPickingButtonChecked() const;
-	void HandlePickingModeStateChanged(ECheckBoxState NewValue);
-	const FSlateBrush* HandleGetPickingModeImage() const;
+	void HandlePickingModeStateChanged();
+	FSlateIcon HandleGetPickingModeImage() const;
 	FText HandleGetPickingModeText() const;
 	TSharedRef<SWidget> HandlePickingModeContextMenu();	
 	void HandlePickButtonClicked(EWidgetPickingMode InPickingMode);
@@ -314,7 +318,7 @@ private:
 	bool IsTakeSnapshotButtonEnabled() const;
 
 	/** Callback for clicking the "Take Snapshot" button. */
-	FReply HandleTakeSnapshotButtonClicked();
+	void HandleTakeSnapshotButtonClicked();
 
 	/** Build option menu for snaphot. */
 	TSharedRef<SWidget> HandleSnapshotOptionsTreeContextMenu();
@@ -330,7 +334,7 @@ private:
 
 #if SLATE_REFLECTOR_HAS_DESKTOP_PLATFORM
 	/** Callback for clicking the "Load Snapshot" button. */
-	FReply HandleLoadSnapshotButtonClicked();
+	void HandleLoadSnapshotButtonClicked();
 #endif // SLATE_REFLECTOR_HAS_DESKTOP_PLATFORM
 
 	/** Called to update the list of available snapshot targets */
@@ -376,6 +380,11 @@ private:
 	/** Should we show only the UMG tree. */
 	bool HandleIsStartTreeWithUMGEnabled() const { return bFilterReflectorTreeRootWithUMG; }
 
+	bool HandleHasPendingActions() const { return !bIsPendingDelayedSnapshot; }
+
+	/** Determine the text of Take Snapshot button */
+	FText HandleGetTakeSnapshotText() const { return bIsPendingDelayedSnapshot ? LOCTEXT("CancelSnapshotButtonText", "Cancel Snapshot") : LOCTEXT("TakeSnapshotButtonText", "Take Snapshot"); }
+
 private:
 	TSharedPtr<FTabManager> TabManager;
 	TMap<FName, TWeakPtr<SDockTab>> SpawnedTabs;
@@ -648,6 +657,8 @@ TSharedRef<SDockTab> SWidgetReflector::SpawnSlateOptionWidgetTab(const FSpawnTab
 
 TSharedRef<SDockTab> SWidgetReflector::SpawnWidgetHierarchyTab(const FSpawnTabArgs& Args)
 {
+	FSlimHorizontalToolBarBuilder ToolbarBuilderGlobal(TSharedPtr<const FUICommandList>(), FMultiBoxCustomization::None);
+
 	TArray<FName> HiddenColumnsList;
 	HiddenColumnsList.Reserve(HiddenReflectorTreeColumns.Num());
 	for (const FString& Item : HiddenReflectorTreeColumns)
@@ -667,7 +678,119 @@ TSharedRef<SDockTab> SWidgetReflector::SpawnWidgetHierarchyTab(const FSpawnTabAr
 		.Text(this, &SWidgetReflector::GetSelectedSnapshotTargetDisplayName)
 	];
 
+	FSlateIcon EmptyIcon(FWidgetReflectorStyle::GetStyleSetName(), "Icon.Empty");
+	ToolbarBuilderGlobal.BeginSection("Picking");
+	{
+
+		FTextBuilder TooltipText;
+		ToolbarBuilderGlobal.AddToolBarButton(
+			FUIAction(
+				FExecuteAction::CreateSP(this, &SWidgetReflector::HandlePickingModeStateChanged),
+				FCanExecuteAction::CreateSP(this,&SWidgetReflector::HandleHasPendingActions),
+				FGetActionCheckState::CreateSP(this, &SWidgetReflector::HandleGetPickingButtonChecked)
+			),
+			NAME_None,
+			MakeAttributeSP(this,&SWidgetReflector::HandleGetPickingModeText),
+			TooltipText.ToText(),
+			MakeAttributeSP(this, &SWidgetReflector::HandleGetPickingModeImage),
+			EUserInterfaceActionType::ToggleButton
+		);
+
+		ToolbarBuilderGlobal.AddComboButton(
+			FUIAction(
+				FExecuteAction(),
+				FCanExecuteAction::CreateSP(this, &SWidgetReflector::HandleHasPendingActions),
+				FGetActionCheckState()
+			),
+			FOnGetContent::CreateSP(this, &SWidgetReflector::HandlePickingModeContextMenu),
+			FText::GetEmpty(),
+			TooltipText.ToText(),
+			EmptyIcon,
+			true
+		);
+
+
+	}
+	ToolbarBuilderGlobal.EndSection();
+
+	ToolbarBuilderGlobal.BeginSection("Filter");
+	{
+
+		FTextBuilder TooltipText;
+
+		ToolbarBuilderGlobal.AddComboButton(
+			FUIAction(
+				FExecuteAction(),
+				FCanExecuteAction::CreateSP(this, &SWidgetReflector::HandleHasPendingActions),
+				FGetActionCheckState()
+			),
+			FOnGetContent::CreateSP(this, &SWidgetReflector::HandleReflectorTreeContextMenu),
+			LOCTEXT("FilterLabel", "Filter"),
+			TooltipText.ToText(),
+			FSlateIcon(FWidgetReflectorStyle::GetStyleSetName(), WidgetReflectorIcon::Filter),
+			false
+		);
+
+	}
+	ToolbarBuilderGlobal.EndSection();
+
+	ToolbarBuilderGlobal.AddWidget(SNew(SSpacer),NAME_None,true,HAlign_Right);
+	ToolbarBuilderGlobal.BeginSection("Option");
+	{
+
+		FTextBuilder TooltipText;
+
+		ToolbarBuilderGlobal.AddToolBarButton(
+			FUIAction(
+				FExecuteAction::CreateSP(this, &SWidgetReflector::HandleTakeSnapshotButtonClicked),
+				FCanExecuteAction::CreateSP(this, &SWidgetReflector::IsTakeSnapshotButtonEnabled),
+				FGetActionCheckState()
+			),
+			NAME_None,
+			MakeAttributeSP(this, &SWidgetReflector::HandleGetTakeSnapshotText),
+			TooltipText.ToText(),
+			FSlateIcon(FWidgetReflectorStyle::GetStyleSetName(), WidgetReflectorIcon::TakeSnapshot),
+			EUserInterfaceActionType::Button
+		);
+
+		ToolbarBuilderGlobal.AddComboButton(
+			FUIAction(
+				FExecuteAction(),
+				FCanExecuteAction::CreateSP(this, &SWidgetReflector::HandleHasPendingActions),
+				FGetActionCheckState()
+			),
+			FOnGetContent::CreateSP(this, &SWidgetReflector::HandleSnapshotOptionsTreeContextMenu),
+			LOCTEXT("OptionsLabel", "Options"),
+			TooltipText.ToText(),
+			FSlateIcon(FWidgetReflectorStyle::GetStyleSetName(), WidgetReflectorIcon::Ellipsis),
+			false
+		);
+
+#if SLATE_REFLECTOR_HAS_DESKTOP_PLATFORM
+
+		ToolbarBuilderGlobal.AddToolBarButton(
+			FUIAction(
+				FExecuteAction::CreateSP(this, &SWidgetReflector::HandleLoadSnapshotButtonClicked),
+				FCanExecuteAction::CreateSP(this, &SWidgetReflector::HandleHasPendingActions),
+				FGetActionCheckState()
+			),
+			NAME_None,
+			LOCTEXT("LoadSnapshotButtonText", "Load Snapshot"),
+			TooltipText.ToText(),
+			FSlateIcon(FWidgetReflectorStyle::GetStyleSetName(), WidgetReflectorIcon::LoadSnapshot),
+			EUserInterfaceActionType::Button
+		);
+
+#endif
+	}
+	ToolbarBuilderGlobal.EndSection();
+
+	ToolbarBuilderGlobal.SetStyle(&FAppStyle::Get(), "SlimToolBar");
+
+
+
 	TSharedRef<SDockTab> SpawnedTab = SNew(SDockTab)
+
 		.Label(LOCTEXT("WidgetHierarchyTab", "Widget Hierarchy"))
 		//.OnCanCloseTab_Lambda([]() { return false; }) // Can't prevent this as it stops the editor from being able to close while the widget reflector is open
 		[
@@ -678,126 +801,9 @@ TSharedRef<SDockTab> SWidgetReflector::SpawnWidgetHierarchyTab(const FSpawnTabAr
 			.Padding(FMargin(0.0f, 2.0f))
 			[
 				SNew(SHorizontalBox)
-
-				+SHorizontalBox::Slot()
-				.AutoWidth()
-				.Padding(FMargin(5.0f, 0.0f))
-				[
-					SNew(SHorizontalBox)
-					+ SHorizontalBox::Slot()
-					.FillWidth(1.f)
-					[
-						SNew(SCheckBox)
-						.Style(FWidgetReflectorStyle::Get(), "CheckBoxNoHover")
-						.Padding(FMargin(4.f, 0.f))
-						.HAlign(HAlign_Left)
-						.IsChecked(this, &SWidgetReflector::HandleGetPickingButtonChecked)
-						.IsEnabled_Lambda([this]() { return !bIsPendingDelayedSnapshot; })
-						.OnCheckStateChanged(this, &SWidgetReflector::HandlePickingModeStateChanged)
-						[
-							SNew(SBox)
-							.MinDesiredWidth(175.f)
-							.VAlign(VAlign_Center)
-							[
-								SNew(SHorizontalBox)
-								+ SHorizontalBox::Slot()
-								.AutoWidth()
-								[
-									SNew(SImage)
-									.Image(this, &SWidgetReflector::HandleGetPickingModeImage)
-								]
-								+ SHorizontalBox::Slot()
-								.AutoWidth()
-								.Padding(10.f, 4.f, 4.f, 4.f)
-								[
-									SNew(STextBlock)
-									.Text(this, &SWidgetReflector::HandleGetPickingModeText)
-								]
-							]
-						]
-					]
-					+ SHorizontalBox::Slot()
-					.AutoWidth()
-					[
-						SNew(SComboButton)
-						.ButtonStyle(FWidgetReflectorStyle::Get(), "Button")
-						.IsEnabled_Lambda([this]() { return !bIsPendingDelayedSnapshot; })
-						.HAlign(HAlign_Center)
-						.VAlign(VAlign_Center)
-						.OnGetMenuContent(this, &SWidgetReflector::HandlePickingModeContextMenu)
-					]
-                ]
-
 				+ SHorizontalBox::Slot()
-				.AutoWidth()
-				[
-					SNew(SComboButton)
-					.ButtonStyle(FWidgetReflectorStyle::Get(), "Button")
-					.IsEnabled_Lambda([this]() { return !bIsPendingDelayedSnapshot; })
-					.HAlign(HAlign_Center)
-					.VAlign(VAlign_Center)
-					.OnGetMenuContent(this, &SWidgetReflector::HandleReflectorTreeContextMenu)
-					.ButtonContent()
-					[
-						SNew(STextBlock)
-						.Text(LOCTEXT("FilterLabel", "Filter "))
-						.ColorAndOpacity(FLinearColor::White)
-					]
-				]
-
-				+SHorizontalBox::Slot()
-				[
-					SNew(SSpacer)
-				]
-
-				+SHorizontalBox::Slot()
-				.AutoWidth()
-				.Padding(FMargin(5.0f, 0.0f))
 				[
-					SNew(SHorizontalBox)
-
-					+SHorizontalBox::Slot()
-					.AutoWidth()
-					[
-						// Button that controls taking a snapshot of the current window(s)
-						SNew(SButton)
-						.VAlign(VAlign_Center)
-						.IsEnabled(this, &SWidgetReflector::IsTakeSnapshotButtonEnabled)
-						.OnClicked(this, &SWidgetReflector::HandleTakeSnapshotButtonClicked)
-						[
-							SNew(STextBlock)
-							.Text_Lambda([this]() { return bIsPendingDelayedSnapshot ? LOCTEXT("CancelSnapshotButtonText", "Cancel Snapshot") : LOCTEXT("TakeSnapshotButtonText", "Take Snapshot"); })
-						]
-					]
-
-					+SHorizontalBox::Slot()
-					.Padding(FMargin(5.0f, 0.0f))
-					.AutoWidth()
-					[
-						SNew(SComboButton)
-						.IsEnabled_Lambda([this]() { return !bIsPendingDelayedSnapshot; })
-						.OnGetMenuContent(this, &SWidgetReflector::HandleSnapshotOptionsTreeContextMenu)
-						.ButtonContent()
-						[
-							SNew(STextBlock)
-							.Text(LOCTEXT("OptionsLabel", "Options"))
-						]
-					]
-#if SLATE_REFLECTOR_HAS_DESKTOP_PLATFORM
-					+ SHorizontalBox::Slot()
-					.AutoWidth()
-					[
-						// Button that controls loading a saved snapshot
-						SNew(SButton)
-						.VAlign(VAlign_Center)
-						.IsEnabled_Lambda([this]() { return !bIsPendingDelayedSnapshot; })
-						.OnClicked(this, &SWidgetReflector::HandleLoadSnapshotButtonClicked)
-						[
-							SNew(STextBlock)
-							.Text(LOCTEXT("LoadSnapshotButtonText", "Load Snapshot"))
-						]
-					]
-#endif // SLATE_REFLECTOR_HAS_DESKTOP_PLATFORM
+					ToolbarBuilderGlobal.MakeWidget()
 				]
 			]
 
@@ -1391,7 +1397,7 @@ ECheckBoxState SWidgetReflector::HandleGetPickingButtonChecked() const
 	return PickingMode != EWidgetPickingMode::None ? ECheckBoxState::Checked : ECheckBoxState::Unchecked;
 }
 
-void SWidgetReflector::HandlePickingModeStateChanged(ECheckBoxState NewValue)
+void SWidgetReflector::HandlePickingModeStateChanged()
 {
 	if (PickingMode == EWidgetPickingMode::None)
 	{
@@ -1408,21 +1414,23 @@ void SWidgetReflector::HandlePickingModeStateChanged(ECheckBoxState NewValue)
 	}
 }
 
-const FSlateBrush* SWidgetReflector::HandleGetPickingModeImage() const
+ FSlateIcon SWidgetReflector::HandleGetPickingModeImage() const
 {
+
+
 	switch (LastPickingMode)
 	{
 	case EWidgetPickingMode::Focus:
-		return FWidgetReflectorStyle::Get().GetBrush(WidgetReflectorIcon::FocusPicking);
+		return FSlateIcon(FWidgetReflectorStyle::GetStyleSetName(), WidgetReflectorIcon::FocusPicking);
 	case EWidgetPickingMode::HitTesting:
-		return FWidgetReflectorStyle::Get().GetBrush(WidgetReflectorIcon::HitTestPicking);
+		return FSlateIcon(FWidgetReflectorStyle::GetStyleSetName(), WidgetReflectorIcon::HitTestPicking);
 	case EWidgetPickingMode::Drawable:
-		return FWidgetReflectorStyle::Get().GetBrush(WidgetReflectorIcon::VisualPicking);
+		return FSlateIcon(FWidgetReflectorStyle::GetStyleSetName(), WidgetReflectorIcon::VisualPicking);
 	case EWidgetPickingMode::None:
 	default:
-		return nullptr;
+		return FSlateIcon(FWidgetReflectorStyle::GetStyleSetName(), "Icon.Empty");
 	}
-	return nullptr;
+	return FSlateIcon(FWidgetReflectorStyle::GetStyleSetName(), "Icon.Empty");
 }
 
 FText SWidgetReflector::HandleGetPickingModeText() const
@@ -1520,7 +1528,7 @@ bool SWidgetReflector::IsTakeSnapshotButtonEnabled() const
 	return SelectedSnapshotTargetInstanceId.IsValid() && !RemoteSnapshotRequestId.IsValid();
 }
 
-FReply SWidgetReflector::HandleTakeSnapshotButtonClicked()
+void SWidgetReflector::HandleTakeSnapshotButtonClicked()
 {
 	if (!bIsPendingDelayedSnapshot)
 	{
@@ -1540,7 +1548,6 @@ FReply SWidgetReflector::HandleTakeSnapshotButtonClicked()
 		TimeOfScheduledSnapshot = -1.0f;
 	}
 
-	return FReply::Handled();
 }
 
 TSharedRef<SWidget> SWidgetReflector::HandleSnapshotOptionsTreeContextMenu()
@@ -1722,7 +1729,7 @@ void SWidgetReflector::HandleRemoteSnapshotReceived(const TArray<uint8>& InSnaps
 
 #if SLATE_REFLECTOR_HAS_DESKTOP_PLATFORM
 
-FReply SWidgetReflector::HandleLoadSnapshotButtonClicked()
+void SWidgetReflector::HandleLoadSnapshotButtonClicked()
 {
 	IDesktopPlatform* DesktopPlatform = FDesktopPlatformModule::Get();
 
@@ -1752,8 +1759,6 @@ FReply SWidgetReflector::HandleLoadSnapshotButtonClicked()
 			WidgetSnapshotVisualizer->SnapshotDataUpdated();
 		}
 	}
-
-	return FReply::Handled();
 }
 
 #endif // SLATE_REFLECTOR_HAS_DESKTOP_PLATFORM
