- Used ToolBarBuilder to generate toolbar in SWidgetSnapshotVisualizer

diff --git a/Engine/Source/Developer/SlateReflector/Private/Widgets/SWidgetSnapshotVisualizer.cpp b/Engine/Source/Developer/SlateReflector/Private/Widgets/SWidgetSnapshotVisualizer.cpp
--- a/Engine/Source/Developer/SlateReflector/Private/Widgets/SWidgetSnapshotVisualizer.cpp
+++ b/Engine/Source/Developer/SlateReflector/Private/Widgets/SWidgetSnapshotVisualizer.cpp
@@ -18,6 +18,7 @@
 
 #include "SlateNavigationEventSimulator.h"
 #include "SlateReflectorModule.h"
+#include "Styling/WidgetReflectorStyle.h"
 
 #if SLATE_REFLECTOR_HAS_DESKTOP_PLATFORM
 #include "DesktopPlatformModule.h"
@@ -794,7 +795,54 @@ void SWidgetSnapshotVisualizer::Construct(const FArguments& InArgs)
 {
 	SnapshotDataPtr = InArgs._SnapshotData;
 	check(SnapshotDataPtr);
+	FSlimHorizontalToolBarBuilder ToolbarBuilderGlobal(TSharedPtr<const FUICommandList>(), FMultiBoxCustomization::None);
+	ToolbarBuilderGlobal.SetStyle(&FAppStyle::Get(), "SlimToolBar");
+	SAssignNew(WindowPickerCombo, SComboBox<TSharedPtr<FWidgetReflectorNodeBase>>)
+		.OptionsSource(&SnapshotDataPtr->GetWindowsPtr())
+		.OnSelectionChanged(this, &SWidgetSnapshotVisualizer::OnWindowSelectionChanged)
+		.OnGenerateWidget(this, &SWidgetSnapshotVisualizer::GenerateWindowPickerComboItem)
+		[
+			SNew(STextBlock)
+			.Text(this, &SWidgetSnapshotVisualizer::GetSelectedWindowComboItemText)
+		]
+		.IsEnabled(this, &SWidgetSnapshotVisualizer::HasValidSnapshot);
+
+	ToolbarBuilderGlobal.BeginSection("Picking");
+	{
+		
+		FTextBuilder TooltipText;
+		ToolbarBuilderGlobal.AddWidget(WindowPickerCombo.ToSharedRef());
+		ToolbarBuilderGlobal.AddToolBarButton(
+			FUIAction(
+				FExecuteAction::CreateSP(this, &SWidgetSnapshotVisualizer::OnPickWidgetClicked),
+				FCanExecuteAction::CreateSP(this, &SWidgetSnapshotVisualizer::HasValidSnapshot),
+				FGetActionCheckState::CreateSP(this, &SWidgetSnapshotVisualizer::GetPickWidgetColor)
+			),
+			NAME_None,
+			MakeAttributeSP(this, &SWidgetSnapshotVisualizer::GetPickWidgetText),
+			TooltipText.ToText(),
+			FSlateIcon(FWidgetReflectorStyle::GetStyleSetName(), "Icon.HitTestPicking"),
+			EUserInterfaceActionType::ToggleButton
+		);
+#if SLATE_REFLECTOR_HAS_DESKTOP_PLATFORM
+		
+		ToolbarBuilderGlobal.AddToolBarButton(
+			FUIAction(
+				FExecuteAction::CreateSP(this, &SWidgetSnapshotVisualizer::OnSaveSnapshotClicked),
+				FCanExecuteAction::CreateSP(this, &SWidgetSnapshotVisualizer::HasValidSnapshot),
+				FGetActionCheckState()
+			),
+			NAME_None,
+			LOCTEXT("SaveSnapshotButtonText", "Save Snapshot"),
+			TooltipText.ToText(),
+			FSlateIcon(FAppStyle::Get().GetStyleSetName(), "Icons.Save"),
+			EUserInterfaceActionType::Button
+		);
+#endif
+
 
+	}
+	ToolbarBuilderGlobal.EndSection();
 	ChildSlot
 	[
 		SNew(SBorder)
@@ -808,45 +856,8 @@ void SWidgetSnapshotVisualizer::Construct(const FArguments& InArgs)
 			.AutoHeight()
 			.Padding(2.f)
 			[
-				SNew(SHorizontalBox)
-
-				+SHorizontalBox::Slot()
-				.AutoWidth()
-				.Padding(FMargin(5.0f, 4.0f))
-				[
-					SAssignNew(WindowPickerCombo, SComboBox<TSharedPtr<FWidgetReflectorNodeBase>>)
-					.OptionsSource(&SnapshotDataPtr->GetWindowsPtr())
-					.OnSelectionChanged(this, &SWidgetSnapshotVisualizer::OnWindowSelectionChanged)
-					.OnGenerateWidget(this, &SWidgetSnapshotVisualizer::GenerateWindowPickerComboItem)
-					[
-						SNew(STextBlock)
-						.Text(this, &SWidgetSnapshotVisualizer::GetSelectedWindowComboItemText)
-					]
-					.IsEnabled(this, &SWidgetSnapshotVisualizer::HasValidSnapshot)
-				]
-
-				+SHorizontalBox::Slot()
-				.AutoWidth()
-				.Padding(FMargin(5.0f, 4.0f))
-				[
-					SNew(SButton)
-					.Text(this, &SWidgetSnapshotVisualizer::GetPickWidgetText)
-					.ButtonColorAndOpacity(this, &SWidgetSnapshotVisualizer::GetPickWidgetColor)
-					.OnClicked(this, &SWidgetSnapshotVisualizer::OnPickWidgetClicked)
-					.IsEnabled(this, &SWidgetSnapshotVisualizer::HasValidSnapshot)
-				]
+				ToolbarBuilderGlobal.MakeWidget()
 
-#if SLATE_REFLECTOR_HAS_DESKTOP_PLATFORM
-				+SHorizontalBox::Slot()
-				.AutoWidth()
-				.Padding(FMargin(5.0f, 4.0f))
-				[
-					SNew(SButton)
-					.Text(LOCTEXT("SaveSnapshotButtonText", "Save Snapshot"))
-					.OnClicked(this, &SWidgetSnapshotVisualizer::OnSaveSnapshotClicked)
-					.IsEnabled(this, &SWidgetSnapshotVisualizer::HasValidSnapshot)
-				]
-#endif // SLATE_REFLECTOR_HAS_DESKTOP_PLATFORM
 			]
 
 			+SVerticalBox::Slot()
@@ -964,23 +975,21 @@ FText SWidgetSnapshotVisualizer::GetPickWidgetText() const
 	return (bIsPicking) ? LOCTEXT("PickingWidget", "Picking (Esc to Stop)") : LOCTEXT("PickSnapshotWidget", "Pick Snapshot Widget");
 }
 
-FSlateColor SWidgetSnapshotVisualizer::GetPickWidgetColor() const
+ECheckBoxState SWidgetSnapshotVisualizer::GetPickWidgetColor() const
 {
-	static const FName SelectionColor("SelectionColor");
-
 	const bool bIsPicking = SnapshotImage.IsValid() && SnapshotImage->GetIsPicking();
 	return bIsPicking
-		? FCoreStyle::Get().GetSlateColor(SelectionColor)
-		: FLinearColor::White;
+		? ECheckBoxState::Checked
+		: ECheckBoxState::Unchecked;
 }
 
-FReply SWidgetSnapshotVisualizer::OnPickWidgetClicked()
+void SWidgetSnapshotVisualizer::OnPickWidgetClicked()
 {
 	if (SnapshotImage.IsValid())
 	{
 		SnapshotImage->SetIsPicking(!SnapshotImage->GetIsPicking());
 	}
-	return FReply::Handled();
+
 }
 
 bool SWidgetSnapshotVisualizer::HasValidSnapshot() const
@@ -999,7 +1008,7 @@ EVisibility SWidgetSnapshotVisualizer::HandleGetNavigationSimulationListVisibili
 
 #if SLATE_REFLECTOR_HAS_DESKTOP_PLATFORM
 
-FReply SWidgetSnapshotVisualizer::OnSaveSnapshotClicked()
+void SWidgetSnapshotVisualizer::OnSaveSnapshotClicked()
 {
 	IDesktopPlatform* DesktopPlatform = FDesktopPlatformModule::Get();
 
@@ -1024,7 +1033,7 @@ FReply SWidgetSnapshotVisualizer::OnSaveSnapshotClicked()
 		}
 	}
 
-	return FReply::Handled();
+
 }
 
 #endif // SLATE_REFLECTOR_HAS_DESKTOP_PLATFORM
diff --git a/Engine/Source/Developer/SlateReflector/Private/Widgets/SWidgetSnapshotVisualizer.h b/Engine/Source/Developer/SlateReflector/Private/Widgets/SWidgetSnapshotVisualizer.h
--- a/Engine/Source/Developer/SlateReflector/Private/Widgets/SWidgetSnapshotVisualizer.h
+++ b/Engine/Source/Developer/SlateReflector/Private/Widgets/SWidgetSnapshotVisualizer.h
@@ -171,11 +171,11 @@ private:
 	/** Get the current text to use for the "Pick Snapshot Widget" button */
 	FText GetPickWidgetText() const;
 
-	/** Get the current color to use for the "Pick Snapshot Widget" button */
-	FSlateColor GetPickWidgetColor() const;
+	/** Get the current check state to use for the "Pick Snapshot Widget" button */
+	ECheckBoxState GetPickWidgetColor() const;
 
 	/** Called when the "Pick Snapshot Widget" button is clicked */
-	FReply OnPickWidgetClicked();
+	void OnPickWidgetClicked();
 
 	/** Is there any snapshot data to save, visualized... */
 	bool HasValidSnapshot() const;
@@ -185,7 +185,7 @@ private:
 
 #if SLATE_REFLECTOR_HAS_DESKTOP_PLATFORM
 	/** Called when the "Save Snapshot" button is clicked */
-	FReply OnSaveSnapshotClicked();
+	void OnSaveSnapshotClicked();
 #endif // SLATE_REFLECTOR_HAS_DESKTOP_PLATFORM
 
 private:
