Fix for no unity error
#rnx

diff --git a/Engine/Source/Developer/SlateReflector/Private/Widgets/SWidgetSnapshotVisualizer.cpp b/Engine/Source/Developer/SlateReflector/Private/Widgets/SWidgetSnapshotVisualizer.cpp
--- a/Engine/Source/Developer/SlateReflector/Private/Widgets/SWidgetSnapshotVisualizer.cpp
+++ b/Engine/Source/Developer/SlateReflector/Private/Widgets/SWidgetSnapshotVisualizer.cpp
@@ -14,6 +14,7 @@
 #include "Widgets/Input/SButton.h"
 #include "Widgets/SNavigationSimulationList.h"
 #include "Framework/Layout/ScrollyZoomy.h"
+#include "Framework/MultiBox/MultiBoxBuilder.h"
 #include "Misc/Base64.h"
 
 #include "SlateNavigationEventSimulator.h"
