Fixed Widget Event Log and Widget Hit-test Grid styles #RB patrick.boutot #jira none

diff --git a/Engine/Source/Developer/SlateReflector/Private/Styling/WidgetReflectorStyle.cpp b/Engine/Source/Developer/SlateReflector/Private/Styling/WidgetReflectorStyle.cpp
--- a/Engine/Source/Developer/SlateReflector/Private/Styling/WidgetReflectorStyle.cpp
+++ b/Engine/Source/Developer/SlateReflector/Private/Styling/WidgetReflectorStyle.cpp
@@ -50,7 +50,7 @@ TSharedRef< FSlateStyleSet > FWidgetReflectorStyle::Create()
 		StyleSet->Set("Icon.HitTestPicking", new FSlateImageBrush(RootToContentDir("Icons/GeneralTools/Select_40x.png"), Icon24x24));
 		StyleSet->Set("Icon.VisualPicking", new FSlateImageBrush(RootToContentDir("Icons/GeneralTools/Paint_40x.png"), Icon24x24));
 		StyleSet->Set("Icon.LoadSnapshot", new FSlateImageBrush(RootToContentDir("Icons/GeneralTools/Import_40x.png"), Icon24x24));
-		StyleSet->Set("Icon.Filter", new FSlateImageBrush(RootToContentDir("Icons/GeneralTools/Filter_40x.png"), Icon24x24));
+		StyleSet->Set("Icon.Filter", new CORE_IMAGE_BRUSH_SVG("Starship/Common/filter", Icon24x24));
 		StyleSet->Set("Icon.TakeSnapshot", new IMAGE_BRUSH_SVG("Starship/Common/SaveThumbnail", Icon24x24));
 
 		StyleSet->Set("Symbols.LeftArrow", new CORE_IMAGE_BRUSH_SVG("Starship/Common/arrow-left", Icon24x24));
diff --git a/Engine/Source/Developer/SlateReflector/Private/Widgets/SWidgetEventLog.cpp b/Engine/Source/Developer/SlateReflector/Private/Widgets/SWidgetEventLog.cpp
--- a/Engine/Source/Developer/SlateReflector/Private/Widgets/SWidgetEventLog.cpp
+++ b/Engine/Source/Developer/SlateReflector/Private/Widgets/SWidgetEventLog.cpp
@@ -41,6 +41,45 @@ void SWidgetEventLog::Construct(const FArguments& InArgs, TSharedPtr<const SWidg
 	}
 	TSharedRef<IMessageLogListing> MessageLogListing = MessageLogModule.GetLogListing(NAME_WidgetEvents);
 
+	FSlimHorizontalToolBarBuilder ToolbarBuilderGlobal(TSharedPtr<const FUICommandList>(), FMultiBoxCustomization::None);
+
+	ToolbarBuilderGlobal.BeginSection("Log");
+	{
+		ToolbarBuilderGlobal.AddWidget(SNew(SCheckBox)
+			.ForegroundColor(FSlateColor::UseForeground())
+			.IsChecked(this, &SWidgetEventLog::HandleFilterWidgetReflectorEventIsChecked)
+			.OnCheckStateChanged(this, &SWidgetEventLog::HandleFilterWidgetReflectorEventStateChanged)
+			[
+				SNew(SBox)
+				.VAlign(VAlign_Center)
+				.HAlign(HAlign_Center)
+				.Padding(FMargin(4.0f, 2.0f))
+				[
+					SNew(STextBlock)
+					.Text(LOCTEXT("FilterWidgetReflectorEvent", "Hide events originated from the Widget Reflector"))
+				]
+			]);
+
+		FTextBuilder TooltipText;
+
+		ToolbarBuilderGlobal.AddComboButton(
+			FUIAction(
+				FExecuteAction(),
+				FCanExecuteAction(),
+				FGetActionCheckState()
+			),
+			FOnGetContent::CreateSP(this, &SWidgetEventLog::OnGenerateCategoriesMenu),
+			LOCTEXT("CategoryComboButtonText", "Filters"),
+			TooltipText.ToText(),
+			FSlateIcon(FWidgetReflectorStyle::GetStyleSetName(), "Icon.Filter"),
+			false
+		);
+
+
+	}
+	ToolbarBuilderGlobal.EndSection();
+
+
 	ChildSlot
 	[
 		SNew(SVerticalBox)
@@ -52,53 +91,7 @@ void SWidgetEventLog::Construct(const FArguments& InArgs, TSharedPtr<const SWidg
 			.BorderImage(FCoreStyle::Get().GetBrush("ToolPanel.GroupBorder"))
 			.Padding(4.f)
 			[
-				SNew(SHorizontalBox)
-
-				+ SHorizontalBox::Slot()
-				.AutoWidth()
-				.HAlign(HAlign_Left)
-				.Padding(4.f, 0.f, 0.f, 0.f)
-				[
-					SNew(SComboButton)
-					.ButtonStyle(FWidgetReflectorStyle::Get(), "Button")
-					.HAlign(HAlign_Center)
-					.VAlign(VAlign_Center)
-					.ContentPadding(FMargin(6.0f, 2.0f))
-					.ButtonContent()
-					[
-						SNew(SBox)
-						.Padding(FMargin(4.f, 0.f, 4.f, 0.f))
-						[
-							SNew(STextBlock)
-							.Text(LOCTEXT("CategoryComboButtonText", "Filters"))
-							.ColorAndOpacity(FLinearColor::White)
-						]
-					]
-					.ContentPadding(FMargin(6.0f, 2.0f))
-					.OnGetMenuContent(this, &SWidgetEventLog::OnGenerateCategoriesMenu)
-				]
-
-				+ SHorizontalBox::Slot()
-				.AutoWidth()
-				.HAlign(HAlign_Left)
-				.Padding(4.f, 0.f, 0.f, 0.f)
-				[
-					SNew(SCheckBox)
-					.Style(FWidgetReflectorStyle::Get(), "CheckBox")
-					.ForegroundColor(FSlateColor::UseForeground())
-					.IsChecked(this, &SWidgetEventLog::HandleFilterWidgetReflectorEventIsChecked)
-					.OnCheckStateChanged(this, &SWidgetEventLog::HandleFilterWidgetReflectorEventStateChanged)
-					[
-						SNew(SBox)
-						.VAlign(VAlign_Center)
-						.HAlign(HAlign_Center)
-						.Padding(FMargin(4.0, 2.0))
-						[
-							SNew(STextBlock)
-							.Text(LOCTEXT("FilterWidgetReflectorEvent", "Hide events originated from the Widget Reflector"))
-						]
-					]
-				]
+				ToolbarBuilderGlobal.MakeWidget()
 			]
 		]
 
diff --git a/Engine/Source/Developer/SlateReflector/Private/Widgets/SWidgetHittestGrid.cpp b/Engine/Source/Developer/SlateReflector/Private/Widgets/SWidgetHittestGrid.cpp
--- a/Engine/Source/Developer/SlateReflector/Private/Widgets/SWidgetHittestGrid.cpp
+++ b/Engine/Source/Developer/SlateReflector/Private/Widgets/SWidgetHittestGrid.cpp
@@ -152,7 +152,6 @@ void SWidgetHittestGrid::Construct(const FArguments& InArgs, TSharedPtr<const SW
 				.Padding(4.f, 0.f, 0.f, 0.f)
 				[
 					SNew(SComboButton)
-					.ButtonStyle(FWidgetReflectorStyle::Get(), "Button")
 					.HAlign(HAlign_Center)
 					.VAlign(VAlign_Center)
 					.ContentPadding(FMargin(6.0f, 2.0f))
@@ -176,7 +175,6 @@ void SWidgetHittestGrid::Construct(const FArguments& InArgs, TSharedPtr<const SW
 				.Padding(4.f, 0.f, 0.f, 0.f)
 				[
 					SNew(SComboButton)
-					.ButtonStyle(FWidgetReflectorStyle::Get(), "Button")
 					.HAlign(HAlign_Center)
 					.VAlign(VAlign_Center)
 					.ContentPadding(FMargin(6.0f, 2.0f))
@@ -199,7 +197,6 @@ void SWidgetHittestGrid::Construct(const FArguments& InArgs, TSharedPtr<const SW
 				.Padding(8.f, 0.f, 0.f, 0.f)
 				[
 					SNew(SCheckBox)
-					.Style(FWidgetReflectorStyle::Get(), "CheckBox")
 					.ForegroundColor(FSlateColor::UseForeground())
 					.ToolTipText(LOCTEXT("PickOnNavigateTooltip", "Attempt to pick the widget in the Widget Reflector when navigation via the hit test grid occurred."))
 					.IsChecked(this, &SWidgetHittestGrid::HandleGetVisualizeOnNavigationChecked)
@@ -221,7 +218,6 @@ void SWidgetHittestGrid::Construct(const FArguments& InArgs, TSharedPtr<const SW
 				.Padding(8.f, 0.f, 0.f, 0.f)
 				[
 					SNew(SCheckBox)
-					.Style(FWidgetReflectorStyle::Get(), "CheckBox")
 					.ForegroundColor(FSlateColor::UseForeground())
 					.ToolTipText(LOCTEXT("RejectWidgetReflectorTooltip", "Reject navigation event originated from the Widget Reflector."))
 					.IsChecked(this, &SWidgetHittestGrid::HandleGetRejectWidgetReflectorChecked)
