- Bug Fix for styling of menu in the widget reflector #RB patrick.boutot #jira none

diff --git a/Engine/Source/Developer/SlateReflector/Private/Styling/WidgetReflectorStyle.cpp b/Engine/Source/Developer/SlateReflector/Private/Styling/WidgetReflectorStyle.cpp
--- a/Engine/Source/Developer/SlateReflector/Private/Styling/WidgetReflectorStyle.cpp
+++ b/Engine/Source/Developer/SlateReflector/Private/Styling/WidgetReflectorStyle.cpp
@@ -69,6 +69,12 @@ TSharedRef< FSlateStyleSet > FWidgetReflectorStyle::Create()
 		StyleSet->Set("BoldSlimToolbar", SlimToolbarStyle);
 	}
 
+	{
+		FToolBarStyle SlimToolbarStyle = FAppStyle::Get().GetWidgetStyle<FToolBarStyle>("SlimToolBar");
+		StyleSet->Set("Menu.Label", SlimToolbarStyle.LabelStyle);
+		StyleSet->Set("Menu.Button", FButtonStyle(SlimToolbarStyle.ButtonStyle).SetNormalPadding(0.0f).SetPressedPadding(0.0f));
+	}
+
 	return StyleSet;
 }
 
