Add Breadcrumb trail for selected node in widget reflector #RB patrick.boutot #jira ue-119462

diff --git a/Engine/Source/Developer/SlateReflector/Private/Widgets/SWidgetReflector.cpp b/Engine/Source/Developer/SlateReflector/Private/Widgets/SWidgetReflector.cpp
--- a/Engine/Source/Developer/SlateReflector/Private/Widgets/SWidgetReflector.cpp
+++ b/Engine/Source/Developer/SlateReflector/Private/Widgets/SWidgetReflector.cpp
@@ -21,6 +21,7 @@
 #include "Widgets/Images/SImage.h"
 #include "Widgets/Text/STextBlock.h"
 #include "Widgets/Layout/SBox.h"
+#include "Widgets/Layout/SSeparator.h"
 #include "Widgets/Input/SButton.h"
 #include "Widgets/SToolTip.h"
 #include "Widgets/Layout/SSplitter.h"
@@ -39,6 +40,7 @@
 #include "Widgets/SInvalidationPanel.h"
 #include "Widgets/Docking/SDockTab.h"
 #include "Widgets/Notifications/SNotificationList.h"
+#include "Widgets/Navigation/SBreadcrumbTrail.h"
 #include "Widgets/Input/SComboBox.h"
 #include "Widgets/SWidgetSnapshotVisualizer.h"
 #include "WidgetSnapshotService.h"
@@ -268,8 +270,11 @@ private:
 	/** Clear previous selection and set the selection to the live widget. */
 	void SelectLiveWidget( TSharedPtr<const SWidget> InWidget );
 
-	/** Set the current selected node as the root of the tree. */
-	void SetSelectedAsReflectorTreeRoot();
+	/** Set the given nodes as the root of the tree. */
+	void SetNodesAsReflectorTreeRoot(TArray<TSharedRef<FWidgetReflectorNodeBase>> RootNodes);
+
+	/** Filter the selected nodes before setting them as root. */
+	TArray<TSharedRef<FWidgetReflectorNodeBase>> FilterSelectedToSetAsReflectorTreeRoot(TArray<TSharedRef<FWidgetReflectorNodeBase>> RootNodes);
 
 	/** Is there any selected node in the reflector tree. */
 	bool DoesReflectorTreeHasSelectedItem() const { return SelectedNodes.Num() > 0; }
@@ -346,6 +351,9 @@ private:
 	/** Called when the list of available snapshot targets changes */
 	void OnAvailableSnapshotTargetsChanged();
 
+	/** Called when a node is set as root to create the breadcrum trail */
+	void CreateCrumbTrailForNode(TSharedRef<FWidgetReflectorNodeBase> InReflectorNode);
+
 	/** Get the display name of the currently selected snapshot target */
 	FText GetSelectedSnapshotTargetDisplayName() const;
 
@@ -368,6 +376,15 @@ private:
 	TSharedRef<SWidget> HandleReflectorTreeContextMenu();
 	TSharedPtr<SWidget> HandleReflectorTreeContextMenuPtr();
 
+	/** Callback for when an item in the reflector tree is clicked on. */
+	void HandleReflectorTreeOnMouseClick(TSharedRef<FWidgetReflectorNodeBase> InReflectorNode);
+
+	/** Callback for when a crumb is clicked on. */
+	void HandleBreadcrumbOnClick(const TSharedRef<FWidgetReflectorNodeBase>& InReflectorNode);
+
+	/** Callback for when the menu of a crumb delimiter is requested. */
+	TSharedRef< SWidget > HandleBreadcrumbDelimiterMenu(const TSharedRef<FWidgetReflectorNodeBase>& InReflectorNode);
+
 	/** Callback for when the reflector tree header list changed. */
 	void HandleReflectorTreeHiddenColumnsListChanged();
 
@@ -377,6 +394,9 @@ private:
 	/** Show the start of the UMG tree. */
 	void HandleStartTreeWithUMG();
 
+	/** Should we display the breadcrumb trail. */
+	EVisibility HandleIsBreadcrumbVisible() const;
+
 	/** Should we show only the UMG tree. */
 	bool HandleIsStartTreeWithUMGEnabled() const { return bFilterReflectorTreeRootWithUMG; }
 
@@ -392,6 +412,7 @@ private:
 	TSharedPtr<SReflectorTree> ReflectorTree;
 	TArray<FString> HiddenReflectorTreeColumns;
 
+	TSharedPtr<SBreadcrumbTrail<TSharedRef<FWidgetReflectorNodeBase>> > BreadCrumb;
 	/** Node that are currently selected */
 	TArray<TSharedRef<FWidgetReflectorNodeBase>> SelectedNodes;
 	/** The original path of the widget picked. It may include node that are now hidden by the filter */
@@ -784,11 +805,8 @@ TSharedRef<SDockTab> SWidgetReflector::SpawnWidgetHierarchyTab(const FSpawnTabAr
 #endif
 	}
 	ToolbarBuilderGlobal.EndSection();
-
 	ToolbarBuilderGlobal.SetStyle(&FAppStyle::Get(), "SlimToolBar");
 
-
-
 	TSharedRef<SDockTab> SpawnedTab = SNew(SDockTab)
 
 		.Label(LOCTEXT("WidgetHierarchyTab", "Widget Hierarchy"))
@@ -806,6 +824,28 @@ TSharedRef<SDockTab> SWidgetReflector::SpawnWidgetHierarchyTab(const FSpawnTabAr
 					ToolbarBuilderGlobal.MakeWidget()
 				]
 			]
+			
+			+ SVerticalBox::Slot()
+			.AutoHeight()
+			.Padding(FMargin(2.0f, 2.0f))
+			[
+				SNew(SVerticalBox)
+				.Visibility(this, &SWidgetReflector::HandleIsBreadcrumbVisible)
+				+ SVerticalBox::Slot()
+				.AutoHeight()
+				[
+					SNew(SSeparator).Thickness(5.f)
+				]
+				+ SVerticalBox::Slot()
+				[
+					SAssignNew(BreadCrumb, SBreadcrumbTrail<TSharedRef<FWidgetReflectorNodeBase>>)
+					.ButtonStyle(FAppStyle::Get(), "SimpleButton")
+					.DelimiterImage(FAppStyle::Get().GetBrush("Icons.ChevronRight"))
+					.TextStyle(FAppStyle::Get(), "NormalText")
+					.OnCrumbClicked(this, &SWidgetReflector::HandleBreadcrumbOnClick)
+					.GetCrumbMenuContent(this, &SWidgetReflector::HandleBreadcrumbDelimiterMenu)
+				]
+			]
 
 			+SVerticalBox::Slot()
 			.FillHeight(1.0f)
@@ -822,6 +862,7 @@ TSharedRef<SDockTab> SWidgetReflector::SpawnWidgetHierarchyTab(const FSpawnTabAr
 					.OnGetChildren(this, &SWidgetReflector::HandleReflectorTreeGetChildren)
 					.OnSelectionChanged(this, &SWidgetReflector::HandleReflectorTreeSelectionChanged)
 					.OnContextMenuOpening(this, &SWidgetReflector::HandleReflectorTreeContextMenuPtr)
+					.OnMouseButtonClick(this, &SWidgetReflector::HandleReflectorTreeOnMouseClick)
 					.HighlightParentNodesForSelection(true)
 					.HeaderRow
 					(
@@ -1075,6 +1116,7 @@ void SWidgetReflector::SetUIMode(const EWidgetReflectorUIMode InNewMode)
 			CloseTab(WidgetReflectorTabID::SnapshotWidgetPicker);
 		}
 	}
+	BreadCrumb->ClearCrumbs();
 }
 
 
@@ -1233,16 +1275,51 @@ void SWidgetReflector::UpdateFilteredTreeRoot()
 	}
 }
 
-void SWidgetReflector::SetSelectedAsReflectorTreeRoot()
+void SWidgetReflector::SetNodesAsReflectorTreeRoot(TArray<TSharedRef<FWidgetReflectorNodeBase>> RootNodes)
 {
-	if (SelectedNodes.Num() > 0)
+	TArray<TSharedRef<FWidgetReflectorNodeBase>> FilteredNodes = FilterSelectedToSetAsReflectorTreeRoot(RootNodes);
+	if (FilteredNodes.Num() > 0)
 	{
 		FilteredTreeRoot.Reset();
-		FilteredTreeRoot.Append(SelectedNodes);
+		FilteredTreeRoot.Append(FilteredNodes);
 		ReflectorTree->RequestTreeRefresh();
+		TSharedPtr<FWidgetReflectorNodeBase> FirstNodeParent = FilteredNodes[0]->GetParentNode();
+		if (FirstNodeParent.IsValid())
+		{
+			CreateCrumbTrailForNode(FirstNodeParent.ToSharedRef());
+		}
+		else
+		{
+			BreadCrumb->ClearCrumbs();
+		}
 	}
 }
 
+TArray<TSharedRef<FWidgetReflectorNodeBase>> SWidgetReflector::FilterSelectedToSetAsReflectorTreeRoot(TArray<TSharedRef<FWidgetReflectorNodeBase>> RootNodes)
+{
+	if (RootNodes.Num() > 1)
+	{
+		TArray<TSharedRef<FWidgetReflectorNodeBase>> ShallowestNodes;
+		ShallowestNodes = RootNodes;
+		for (int32 index = ShallowestNodes.Num() -1 ; index >= 0; index--)
+		{
+			if (ShallowestNodes.Contains(ShallowestNodes[index]->GetParentNode()))
+			{
+				ShallowestNodes.RemoveAt(index);
+			}
+		}
+		TSharedPtr<FWidgetReflectorNodeBase> FirstNodeParent = ShallowestNodes[0]->GetParentNode();
+		for (int32 index = ShallowestNodes.Num() - 1; index >= 0; index--)
+		{
+			if (ShallowestNodes[index]->GetParentNode() != FirstNodeParent)
+			{
+				ShallowestNodes.RemoveAt(index);
+			}
+		}
+		return ShallowestNodes;
+	}
+	return RootNodes;
+}
 TSharedPtr<IToolTip> SWidgetReflector::GenerateToolTipForReflectorNode( TWeakPtr<FWidgetReflectorNodeBase> InReflectorNode ) const
 {
 	if (TSharedPtr<FWidgetReflectorNodeBase> ReflectorNode = InReflectorNode.Pin())
@@ -1941,7 +2018,7 @@ TSharedRef<SWidget> SWidgetReflector::HandleReflectorTreeContextMenu()
 		LOCTEXT("SetAsRootTooltip", "Set selected node as the root of the graph"),
 		FSlateIcon(),
 		FUIAction(
-			FExecuteAction::CreateSP(this, &SWidgetReflector::SetSelectedAsReflectorTreeRoot),
+			FExecuteAction::CreateSP(this, &SWidgetReflector::SetNodesAsReflectorTreeRoot, SelectedNodes),
 			FCanExecuteAction::CreateSP(this, &SWidgetReflector::DoesReflectorTreeHasSelectedItem)
 		));
 
@@ -1991,9 +2068,80 @@ void SWidgetReflector::HandleReflectorTreeHiddenColumnsListChanged()
 	}
 #endif
 }
+void SWidgetReflector::CreateCrumbTrailForNode(TSharedRef<FWidgetReflectorNodeBase> InReflectorNode)
+{
+	FWidgetPath PathToWidget;
+	TSharedRef<SWidget> SelectedNodeAsWidget = InReflectorNode.Get().GetLiveWidget().ToSharedRef();
+	TArray<TSharedRef<FWidgetReflectorNodeBase>> PickedNodePath;
+	FWidgetReflectorNodeUtils::FindLiveWidget(ReflectorTreeRoot, SelectedNodeAsWidget, PickedNodePath);
+	BreadCrumb->ClearCrumbs();
+	if (PickedNodePath.Num() > 1)
+	{
+		for (int32 i = 0; i < PickedNodePath.Num(); i++)
+		{
+			TSharedPtr<SWidget> Parent = PickedNodePath[i]->GetLiveWidget();
+			FText WidgetType = FWidgetReflectorNodeUtils::GetWidgetType(Parent);
+			BreadCrumb->PushCrumb(WidgetType, PickedNodePath[i]);
+		}
+	}
+}
+void SWidgetReflector::HandleReflectorTreeOnMouseClick(TSharedRef<FWidgetReflectorNodeBase> InReflectorNode)
+{
+	const FModifierKeysState ModKeyState = FSlateApplication::Get().GetModifierKeys();
+
+	if (ModKeyState.IsLeftAltDown())
+	{
+		if (SelectedNodes.Num() > 0)
+		{
+			TArray<TSharedRef<FWidgetReflectorNodeBase>> RootNodes;
+			RootNodes.Add(InReflectorNode);
+			SetNodesAsReflectorTreeRoot(RootNodes);
+		}
+	}
+}
+
+void SWidgetReflector::HandleBreadcrumbOnClick(const TSharedRef<FWidgetReflectorNodeBase>& InReflectorNode)
+{
+	FilteredTreeRoot.Reset();
+	FilteredTreeRoot.Add(InReflectorNode);
+	BreadCrumb->PopCrumb();
+	ReflectorTree->RequestTreeRefresh();
+}
+
+TSharedRef< SWidget > SWidgetReflector::HandleBreadcrumbDelimiterMenu(const TSharedRef<FWidgetReflectorNodeBase>& InReflectorNode)
+{
+	const bool bShouldCloseWindowAfterMenuSelection = true;
+	FMenuBuilder MenuBuilder(bShouldCloseWindowAfterMenuSelection, nullptr);
+
+	bool bHasFilteredTreeRoot = ReflectorTreeRoot != FilteredTreeRoot;
+	TArray<TSharedRef< FWidgetReflectorNodeBase>> Children = InReflectorNode->GetChildNodes();
+	for (int32 ChildIndex = 0; ChildIndex < Children.Num(); ++ChildIndex)
+	{
+		TSharedPtr<SWidget> ChildWidget = Children[ChildIndex]->GetLiveWidget();
+		FText WidgetType = FWidgetReflectorNodeUtils::GetWidgetType(ChildWidget);
+		TArray<TSharedRef<FWidgetReflectorNodeBase>> RootNodes;
+		RootNodes.Add(Children[ChildIndex]);
+		MenuBuilder.AddMenuEntry(
+			WidgetType,
+			FText::GetEmpty(),
+			FSlateIcon(),
+			FUIAction(
+				FExecuteAction::CreateSP(this, &SWidgetReflector::SetNodesAsReflectorTreeRoot, RootNodes)
+			
+			));
+	}
+
+	return MenuBuilder.MakeWidget();
+}
+
+EVisibility SWidgetReflector::HandleIsBreadcrumbVisible() const
+{
+	return BreadCrumb->HasCrumbs() ? EVisibility::Visible : EVisibility::Collapsed;
+}
 
 void SWidgetReflector::HandleResetFilteredTreeRoot()
 {
+	BreadCrumb->ClearCrumbs();
 	bFilterReflectorTreeRootWithUMG = false;
 	UpdateFilteredTreeRoot();
 	ReflectorTree->RequestTreeRefresh();
