WidgetReflector: Add the LayerId

diff --git a/Engine/Source/Developer/SlateReflector/Private/Models/WidgetReflectorNode.cpp b/Engine/Source/Developer/SlateReflector/Private/Models/WidgetReflectorNode.cpp
--- a/Engine/Source/Developer/SlateReflector/Private/Models/WidgetReflectorNode.cpp
+++ b/Engine/Source/Developer/SlateReflector/Private/Models/WidgetReflectorNode.cpp
@@ -145,6 +145,11 @@ FText FLiveWidgetReflectorNode::GetWidgetClippingText() const
 	return FWidgetReflectorNodeUtils::GetWidgetClippingText(Widget.Pin());
 }
 
+int32 FLiveWidgetReflectorNode::GetWidgetLayerId() const
+{
+	return FWidgetReflectorNodeUtils::GetWidgetLayerId(Widget.Pin());
+}
+
 bool FLiveWidgetReflectorNode::GetWidgetFocusable() const
 {
 	return FWidgetReflectorNodeUtils::GetWidgetFocusable(Widget.Pin());
@@ -269,6 +274,7 @@ FSnapshotWidgetReflectorNode::FSnapshotWidgetReflectorNode(const FArrangedWidget
 	, bCachedWidgetIsInvalidationRoot(FWidgetReflectorNodeUtils::GetWidgetIsInvalidationRoot(InArrangedWidget.Widget))
 	, bCachedWidgetEnabled(FWidgetReflectorNodeUtils::GetWidgetEnabled(InArrangedWidget.Widget))
 	, CachedWidgetClippingText(FWidgetReflectorNodeUtils::GetWidgetClippingText(InArrangedWidget.Widget))
+	, CachedWidgetLayerId(FWidgetReflectorNodeUtils::GetWidgetLayerId(InArrangedWidget.Widget))
 	, CachedWidgetReadableLocation(FWidgetReflectorNodeUtils::GetWidgetReadableLocation(InArrangedWidget.Widget))
 	, CachedWidgetFile(FWidgetReflectorNodeUtils::GetWidgetFile(InArrangedWidget.Widget))
 	, CachedWidgetLineNumber(FWidgetReflectorNodeUtils::GetWidgetLineNumber(InArrangedWidget.Widget))
@@ -326,6 +332,11 @@ FText FSnapshotWidgetReflectorNode::GetWidgetClippingText() const
 	return CachedWidgetClippingText;
 }
 
+int32 FSnapshotWidgetReflectorNode::GetWidgetLayerId() const
+{
+	return CachedWidgetLayerId;
+}
+
 bool FSnapshotWidgetReflectorNode::GetWidgetNeedsTick() const
 {
 	return bCachedWidgetNeedsTick;
@@ -506,6 +517,7 @@ TSharedRef<FJsonValue> FSnapshotWidgetReflectorNode::ToJson(const TSharedRef<FSn
 	RootJsonObject->SetBoolField(TEXT("WidgetIsInvalidationRoot"), RootSnapshotNode->bCachedWidgetIsInvalidationRoot);
 	RootJsonObject->SetBoolField(TEXT("WidgetEnabled"), RootSnapshotNode->bCachedWidgetEnabled);
 	RootJsonObject->SetStringField(TEXT("WidgetClippingText"), RootSnapshotNode->CachedWidgetClippingText.ToString());
+	RootJsonObject->SetNumberField(TEXT("WidgetLayerId"), RootSnapshotNode->CachedWidgetLayerId);
 	RootJsonObject->SetStringField(TEXT("WidgetReadableLocation"), RootSnapshotNode->CachedWidgetReadableLocation.ToString());
 	RootJsonObject->SetStringField(TEXT("WidgetFile"), RootSnapshotNode->CachedWidgetFile);
 	RootJsonObject->SetNumberField(TEXT("WidgetLineNumber"), RootSnapshotNode->CachedWidgetLineNumber);
@@ -678,6 +690,7 @@ TSharedRef<FSnapshotWidgetReflectorNode> FSnapshotWidgetReflectorNode::FromJson(
 	RootSnapshotNode->bCachedWidgetIsInvalidationRoot = RootJsonObject->GetBoolField(TEXT("WidgetIsInvalidationRoot"));
 	RootSnapshotNode->bCachedWidgetEnabled = RootJsonObject->GetBoolField(TEXT("WidgetEnabled"));
 	RootSnapshotNode->CachedWidgetClippingText = FText::FromString(RootJsonObject->GetStringField(TEXT("WidgetClippingText")));
+	RootSnapshotNode->CachedWidgetLayerId = RootJsonObject->GetIntegerField(TEXT("WidgetLayerId"));
 	RootSnapshotNode->CachedWidgetReadableLocation = FText::FromString(RootJsonObject->GetStringField(TEXT("WidgetReadableLocation")));
 	RootSnapshotNode->CachedWidgetFile = RootJsonObject->GetStringField(TEXT("WidgetFile"));
 	RootSnapshotNode->CachedWidgetLineNumber = RootJsonObject->GetIntegerField(TEXT("WidgetLineNumber"));
@@ -999,6 +1012,11 @@ FText FWidgetReflectorNodeUtils::GetWidgetClippingText(const TSharedPtr<const SW
 	return FText::GetEmpty();
 }
 
+int32 FWidgetReflectorNodeUtils::GetWidgetLayerId(const TSharedPtr<const SWidget>& InWidget)
+{
+	return (InWidget.IsValid()) ? InWidget->GetPersistentState().LayerId : -1;
+}
+
 FText FWidgetReflectorNodeUtils::GetWidgetReadableLocation(const TSharedPtr<const SWidget>& InWidget)
 {
 	return FText::FromString(FReflectionMetaData::GetWidgetDebugInfo(InWidget.Get()));
diff --git a/Engine/Source/Developer/SlateReflector/Private/Models/WidgetReflectorNode.h b/Engine/Source/Developer/SlateReflector/Private/Models/WidgetReflectorNode.h
--- a/Engine/Source/Developer/SlateReflector/Private/Models/WidgetReflectorNode.h
+++ b/Engine/Source/Developer/SlateReflector/Private/Models/WidgetReflectorNode.h
@@ -76,6 +76,11 @@ public:
 	 */
 	virtual FText GetWidgetClippingText() const = 0;
 
+	/**
+	 * @return The LayerId the widget were painted from
+	 */
+	virtual int32 GetWidgetLayerId() const = 0;
+
 	/**
 	 * @return The bool indicating whether or not the widget we were initialized from reports as Focusable
 	 */
@@ -281,6 +286,7 @@ public:
 	virtual bool GetWidgetVisible() const override;
 	virtual bool GetWidgetVisibilityInherited() const override;
 	virtual FText GetWidgetClippingText() const override;
+	virtual int32 GetWidgetLayerId() const override;
 	virtual bool GetWidgetFocusable() const override;
 	virtual bool GetWidgetNeedsTick() const override;
 	virtual bool GetWidgetIsVolatile() const override;
@@ -339,6 +345,7 @@ public:
 	virtual FText GetWidgetTypeAndShortName() const override;
 	virtual FText GetWidgetVisibilityText() const override;
 	virtual FText GetWidgetClippingText() const override;
+	virtual int32 GetWidgetLayerId() const override;
 	virtual bool GetWidgetVisible() const override;
 	virtual bool GetWidgetVisibilityInherited() const override;
 	virtual bool GetWidgetFocusable() const override;
@@ -417,6 +424,9 @@ private:
 	/** The clipping string of the widget at the point it was passed to Initialize */
 	FText CachedWidgetClippingText;
 
+	/** The LayerId of the widget */
+	int32 CachedWidgetLayerId;
+
 	/** The human readable location (source file for C++ widgets, asset name for UMG widgets) of the widget at the point it was passed to Initialize */
 	FText CachedWidgetReadableLocation;
 
@@ -559,6 +569,11 @@ public:
 	 */
 	static FText GetWidgetClippingText(const TSharedPtr<const SWidget>& InWidget);
 
+	/**
+	 * @return The current LayerId for the given widget
+	 */
+	static int32 GetWidgetLayerId(const TSharedPtr<const SWidget>& InWidget);
+
 	/**
 	* @return The current focusability for the given widget
 	*/
diff --git a/Engine/Source/Developer/SlateReflector/Private/Widgets/SWidgetReflector.cpp b/Engine/Source/Developer/SlateReflector/Private/Widgets/SWidgetReflector.cpp
--- a/Engine/Source/Developer/SlateReflector/Private/Widgets/SWidgetReflector.cpp
+++ b/Engine/Source/Developer/SlateReflector/Private/Widgets/SWidgetReflector.cpp
@@ -466,6 +466,7 @@ void SWidgetReflector::Construct( const FArguments& InArgs )
 	HiddenReflectorTreeColumns.Add(SReflectorTreeWidgetItem::NAME_Volatile.ToString());
 	HiddenReflectorTreeColumns.Add(SReflectorTreeWidgetItem::NAME_HasActiveTimer.ToString());
 	HiddenReflectorTreeColumns.Add(SReflectorTreeWidgetItem::NAME_ActualSize.ToString());
+	HiddenReflectorTreeColumns.Add(SReflectorTreeWidgetItem::NAME_LayerId.ToString());
 
 	LoadSettings();
 
@@ -905,6 +906,10 @@ TSharedRef<SDockTab> SWidgetReflector::SpawnWidgetHierarchyTab(const FSpawnTabAr
 						.DefaultLabel(LOCTEXT("Clipping", "Clipping" ))
 						.FillSized(100.0f)
 
+						+ SHeaderRow::Column(SReflectorTreeWidgetItem::NAME_LayerId)
+						.DefaultLabel(LOCTEXT("LayerId", "LayerId"))
+						.FillSized(35.f)
+
 						+ SHeaderRow::Column(SReflectorTreeWidgetItem::NAME_ActualSize)
 						.DefaultLabel(LOCTEXT("ActualSize", "Size"))
 						.FillSized(100.0f)
diff --git a/Engine/Source/Developer/SlateReflector/Private/Widgets/SWidgetReflectorTreeWidgetItem.cpp b/Engine/Source/Developer/SlateReflector/Private/Widgets/SWidgetReflectorTreeWidgetItem.cpp
--- a/Engine/Source/Developer/SlateReflector/Private/Widgets/SWidgetReflectorTreeWidgetItem.cpp
+++ b/Engine/Source/Developer/SlateReflector/Private/Widgets/SWidgetReflectorTreeWidgetItem.cpp
@@ -20,6 +20,7 @@ FName SReflectorTreeWidgetItem::NAME_Enabled(TEXT("Enabled"));
 FName SReflectorTreeWidgetItem::NAME_Volatile(TEXT("Volatile"));
 FName SReflectorTreeWidgetItem::NAME_HasActiveTimer(TEXT("HasActiveTimer"));
 FName SReflectorTreeWidgetItem::NAME_Clipping(TEXT("Clipping"));
+FName SReflectorTreeWidgetItem::NAME_LayerId(TEXT("LayerId"));
 FName SReflectorTreeWidgetItem::NAME_ForegroundColor(TEXT("ForegroundColor"));
 FName SReflectorTreeWidgetItem::NAME_ActualSize(TEXT("ActualSize"));
 FName SReflectorTreeWidgetItem::NAME_Address(TEXT("Address"));
@@ -140,6 +141,17 @@ TSharedRef<SWidget> SReflectorTreeWidgetItem::GenerateWidgetForColumn(const FNam
 				.Text(WidgetInfo->GetWidgetClippingText())
 			];
 	}
+	else if (ColumnName == NAME_LayerId)
+	{
+		return SNew(SBox)
+			.HAlign(HAlign_Right)
+			.VAlign(VAlign_Center)
+			.Padding(FMargin(2.0f, 0.0f))
+			[
+				SNew(STextBlock)
+				.Text(FText::AsNumber(WidgetInfo->GetWidgetLayerId()))
+			];
+	}
 	else if (ColumnName == NAME_ForegroundColor )
 	{
 		const FSlateColor Foreground = WidgetInfo->GetWidgetForegroundColor();
diff --git a/Engine/Source/Developer/SlateReflector/Private/Widgets/SWidgetReflectorTreeWidgetItem.h b/Engine/Source/Developer/SlateReflector/Private/Widgets/SWidgetReflectorTreeWidgetItem.h
--- a/Engine/Source/Developer/SlateReflector/Private/Widgets/SWidgetReflectorTreeWidgetItem.h
+++ b/Engine/Source/Developer/SlateReflector/Private/Widgets/SWidgetReflectorTreeWidgetItem.h
@@ -27,6 +27,7 @@ public:
 	static FName NAME_Volatile;
 	static FName NAME_HasActiveTimer;
 	static FName NAME_Clipping;
+	static FName NAME_LayerId;
 	static FName NAME_ForegroundColor;
 	static FName NAME_Address;
 	static FName NAME_ActualSize;
diff --git a/Engine/Source/Developer/SlateReflector/Private/Widgets/SWidgetSnapshotVisualizer.cpp b/Engine/Source/Developer/SlateReflector/Private/Widgets/SWidgetSnapshotVisualizer.cpp
--- a/Engine/Source/Developer/SlateReflector/Private/Widgets/SWidgetSnapshotVisualizer.cpp
+++ b/Engine/Source/Developer/SlateReflector/Private/Widgets/SWidgetSnapshotVisualizer.cpp
@@ -448,7 +448,7 @@ void FWidgetSnapshotData::SaveSnapshotToBuffer(TArray<uint8>& OutData) const
 	BufferWriter.SerializeCompressed(TmpJsonData.GetData(), TmpJsonData.Num(), NAME_Zlib);
 }
 
-double SnapshotJsonVersion = 2.1;
+double SnapshotJsonVersion = 2.2;
 
 TSharedRef<FJsonObject> FWidgetSnapshotData::SaveSnapshotAsJson() const
 {
