Added Widget List to Widget Reflector, which lists all currently active widgets, grouped by creation location. This can be useful for pinpointing SWidget leaks.

This requires UE_WITH_SLATE_DEBUG_WIDGETLIST to be defined, since it makes use of the functionality in FWidgetList.

diff --git a/Engine/Source/Developer/SlateReflector/Private/Widgets/SWidgetList.cpp b/Engine/Source/Developer/SlateReflector/Private/Widgets/SWidgetList.cpp
new file mode 100644
--- /dev/null
+++ b/Engine/Source/Developer/SlateReflector/Private/Widgets/SWidgetList.cpp
@@ -0,0 +1,360 @@
+// Copyright Epic Games, Inc. All Rights Reserved.
+
+#include "SWidgetList.h"
+
+#include "DesktopPlatformModule.h"
+#include "Framework/Application/SlateApplication.h"
+#include "HAL/PlatformProcess.h"
+#include "Internationalization/Internationalization.h"
+#include "Internationalization/Text.h"
+#include "SPositiveActionButton.h"
+#include "SSimpleButton.h"
+#include "Styling/StyleColors.h"
+#include "Types/ReflectionMetadata.h"
+#include "Widgets/Input/SHyperlink.h"
+#include "Widgets/Layout/SUniformGridPanel.h"
+#include "Widgets/Views/STableRow.h"
+
+#define LOCTEXT_NAMESPACE "WidgetList"
+
+struct FWidgetListEntry
+{
+	FName Type;
+	FName Location;
+	FAssetData AssetData;
+	uint64 Count = 0;
+};
+
+class SWidgetListRow : public SMultiColumnTableRow<TSharedPtr<FWidgetListEntry>>
+{
+public:
+	SLATE_BEGIN_ARGS(SWidgetListRow) {}
+		SLATE_ARGUMENT(FAccessSourceCode, OnAccessSource)
+		SLATE_ARGUMENT(FAccessAsset, OnAccessAsset)
+	SLATE_END_ARGS()
+
+	static const FName NAME_Type;
+	static const FName NAME_Location;
+	static const FName NAME_Count;
+
+	void Construct(const FArguments& InArgs, const TSharedPtr<FWidgetListEntry>& InEntry, const TSharedRef<STableViewBase>& OwnerTable)
+	{
+		Entry = InEntry;
+		OnAccessSource = InArgs._OnAccessSource;
+		OnAccessAsset = InArgs._OnAccessAsset;
+
+		FSuperRowType::FArguments SuperArgs = FSuperRowType::FArguments();
+		FSuperRowType::Construct(SuperArgs, OwnerTable);
+	}
+
+	virtual TSharedRef<SWidget> GenerateWidgetForColumn(const FName& ColumnName) override
+	{
+		if (ColumnName == NAME_Type)
+		{
+			return SNew(STextBlock)
+				.Text(FText::FromName(Entry->Type));
+		}
+		if (ColumnName == NAME_Location)
+		{
+			const FString ReadableLocation = FString::Printf(TEXT("%s(%d)"), *FPaths::GetCleanFilename(Entry->Location.GetPlainNameString()), Entry->Location.GetNumber());
+			return SNew(SBox)
+				.HAlign(HAlign_Left)
+				.VAlign(VAlign_Center)
+				.Padding(FMargin(2.0f, 0.0f))
+				[
+					SNew(SHyperlink)
+					.Text(FText::FromString(ReadableLocation))
+					.OnNavigate(this, &SWidgetListRow::OnNavigateToSource)
+				];
+		}
+		else if (ColumnName == NAME_Count)
+		{
+			return SNew(STextBlock)
+				.Text(FText::AsNumber(Entry->Count));
+		}
+
+		return SNullWidget::NullWidget;
+	}
+
+private:
+	void OnNavigateToSource()
+	{
+		if (OnAccessAsset.IsBound() && Entry->AssetData.IsValid())
+		{
+			Entry->AssetData.GetPackage();
+			OnAccessAsset.Execute(Entry->AssetData.GetAsset());
+		}
+
+		if (OnAccessAsset.IsBound())
+		{
+			OnAccessSource.Execute(Entry->Location.GetPlainNameString(), Entry->Location.GetNumber(), 0);
+		}
+	}
+
+private:
+	TSharedPtr<FWidgetListEntry> Entry;
+	FAccessSourceCode OnAccessSource;
+	FAccessAsset OnAccessAsset;
+};
+
+const FName SWidgetListRow::NAME_Type = "Type";
+const FName SWidgetListRow::NAME_Location = "Location";
+const FName SWidgetListRow::NAME_Count = "Count";
+
+void SWidgetList::Construct(const FArguments& Args)
+{
+	OnAccessAsset = Args._OnAccessAsset;
+	OnAccessSource = Args._OnAccessSource;
+
+#if UE_WITH_SLATE_DEBUG_WIDGETLIST
+	ChildSlot
+	[
+		SNew(SVerticalBox)
+		+ SVerticalBox::Slot()
+		.Padding(4, 6)
+		.VAlign(VAlign_Top)
+		.AutoHeight()
+		[
+			SNew(SHorizontalBox)
+			+ SHorizontalBox::Slot()
+			.HAlign(HAlign_Fill)
+			.VAlign(VAlign_Center)
+			[
+				SNew(SUniformGridPanel)
+				.SlotPadding(FMargin(0, 0, 10, 0))
+				+ SUniformGridPanel::Slot(0, 0)
+				[
+					SNew(STextBlock)
+					.Text_Lambda([this]() { return FText::Format(LOCTEXT("CurrentFormat", "Current: {0}"), CurrentCount); })
+				]
+				+ SUniformGridPanel::Slot(1, 0)
+				[
+					SNew(STextBlock)
+					.Text_Lambda([this]() { return FText::Format(LOCTEXT("PreviousFormat", "Previous: {0}"), PreviousCount); })
+				]
+				+ SUniformGridPanel::Slot(2, 0)
+				[
+					SNew(STextBlock)
+					.Text_Lambda([this]() 
+					{
+						FNumberFormattingOptions Options = FNumberFormattingOptions().SetAlwaysSign(true);
+						return FText::Format(LOCTEXT("CurrentFormat", "Current: {0}"), FText::AsNumber(CurrentCount - PreviousCount, &Options));
+					})
+					.ColorAndOpacity_Lambda([this]()
+					{
+						const int32 Diff = CurrentCount - PreviousCount;
+						return Diff > 0 ? FStyleColors::AccentGreen :
+							(Diff < 0 ? FStyleColors::AccentRed :
+							FStyleColors::Foreground); 
+					})
+				]
+			]
+			+ SHorizontalBox::Slot()
+			.HAlign(HAlign_Right)
+			.VAlign(VAlign_Center)
+			.Padding(0,0,4,0)
+			.AutoWidth()
+			[
+				SNew(SSimpleButton)
+				.Icon(FAppStyle::Get().GetBrush("Icons.Save"))
+				.Text(LOCTEXT("WidgetListExport", "Export to CSV..."))
+				.OnClicked(this, &SWidgetList::ExportToCSV)
+			]
+			+ SHorizontalBox::Slot()
+			.HAlign(HAlign_Right)
+			.VAlign(VAlign_Center)
+			.AutoWidth()
+			[
+				SNew(SPositiveActionButton)
+				.Icon(FAppStyle::Get().GetBrush("Icons.Refresh"))
+				.Text(LOCTEXT("WidgetListRefresh", "Refresh"))
+				.OnClicked_Lambda([this]() { Refresh(); return FReply::Handled(); })
+			]
+		]
+		+ SVerticalBox::Slot()
+		[
+			SAssignNew(ListView, SListView<TSharedPtr<FWidgetListEntry>>)
+			.ItemHeight(24)
+			.ListItemsSource(&Entries)
+			.OnGenerateRow(this, &SWidgetList::OnGenerateRow)
+			.HeaderRow
+			(
+				SNew(SHeaderRow)
+				+ SHeaderRow::Column(SWidgetListRow::NAME_Type)
+				.FillWidth(0.3)
+				.DefaultLabel(LOCTEXT("WidgetListWidget", "Type"))
+				.OnSort(this, &SWidgetList::OnSort)
+
+				+ SHeaderRow::Column(SWidgetListRow::NAME_Location)
+				.FillWidth(0.5)
+				.DefaultLabel(LOCTEXT("WidgetListLocation", "Location"))
+				.OnSort(this, &SWidgetList::OnSort)
+
+				+ SHeaderRow::Column(SWidgetListRow::NAME_Count)
+				.FillWidth(0.2)
+				.DefaultLabel(LOCTEXT("WidgetListCount", "Count"))
+				.OnSort(this, &SWidgetList::OnSort)
+			)
+		]
+	];
+#else
+
+	ChildSlot
+	[
+		SNew(SBox)
+		.Padding(10)
+		.HAlign(HAlign_Center)
+		.VAlign(VAlign_Center)
+		[
+			SNew(STextBlock)
+			.Text(LOCTEXT("WidgetListNotDefined", "Widget list requires UE_WITH_SLATE_DEBUG_WIDGETLIST to be defined."))
+			.ColorAndOpacity(FStyleColors::Error)
+		]
+	];
+
+#endif // UE_WITH_SLATE_DEBUG_WIDGETLIST
+}
+
+void SWidgetList::Refresh()
+{
+	TArray<const SWidget*> Widgets = FSlateDebugging::GetAllWidgets();
+
+	PreviousCount = CurrentCount;
+	CurrentCount = Widgets.Num();
+
+	using FWidgetKey = TPair<FName, FName>;
+	TMap<FWidgetKey, TSharedPtr<FWidgetListEntry>> WidgetEntryMap;
+
+	for (const SWidget* Widget : Widgets)
+	{
+		FWidgetKey Key { Widget->GetType(), Widget->GetCreatedInLocation() };
+		TSharedPtr<FWidgetListEntry>* Entry = WidgetEntryMap.Find(Key);
+		if (Entry == nullptr)
+		{
+			Entry = &WidgetEntryMap.Add(Key, MakeShared<FWidgetListEntry>());
+			(*Entry)->Type = Widget->GetType();
+			(*Entry)->Location = Widget->GetCreatedInLocation();
+			(*Entry)->Count = 0;
+
+			TSharedPtr<FReflectionMetaData> WidgetMeta = Widget->GetMetaData<FReflectionMetaData>();
+			if (WidgetMeta.IsValid() && WidgetMeta->Asset.Get() != nullptr)
+			{
+				(*Entry)->AssetData = FAssetData(WidgetMeta->Asset.Get());
+			}
+		}
+
+		(*Entry)->Count++;
+	}
+
+	Entries.Reset();
+
+	for (const TPair<FWidgetKey, TSharedPtr<FWidgetListEntry>>& EntryPair : WidgetEntryMap)
+	{
+		Entries.Add(EntryPair.Value);
+	}
+
+	Sort();
+}
+
+FReply SWidgetList::ExportToCSV() const
+{
+	FString OutputFile;
+	const TCHAR* DefaultFile = TEXT("WidgetList.csv");
+
+	IDesktopPlatform* DesktopPlatform = FDesktopPlatformModule::Get();
+	if (DesktopPlatform)
+	{
+		const void* ParentWindowHandle = FSlateApplication::Get().FindBestParentWindowHandleForDialogs(nullptr);
+
+		TArray<FString> SaveFilenames;
+
+		bool bSave = DesktopPlatform->SaveFileDialog(
+			ParentWindowHandle,
+			LOCTEXT("ExportToCSVTitle", "Export...").ToString(),
+			FPaths::ProjectSavedDir(),
+			DefaultFile,
+			TEXT("CSV|*.csv"),
+			EFileDialogFlags::None,
+			SaveFilenames
+		);
+
+		if (!bSave || SaveFilenames.Num() == 0)
+		{
+			return FReply::Handled();
+		}
+
+		OutputFile = SaveFilenames[0];
+	}
+	else
+	{
+		OutputFile = FPaths::Combine(FPaths::ProjectSavedDir(), TEXT("WidgetList.csv"));
+	}
+
+	FSlateDebugging::ExportWidgetList(OutputFile);
+
+	FPlatformProcess::ExploreFolder(*OutputFile);
+
+	return FReply::Handled();
+}
+
+void SWidgetList::Sort()
+{
+	if (SortColumn.IsNone())
+	{
+		SortColumn = SWidgetListRow::NAME_Count;
+	}
+
+	if (SortColumn.IsEqual(SWidgetListRow::NAME_Count))
+	{
+		Entries.Sort([this](const TSharedPtr<FWidgetListEntry>& A, const TSharedPtr<FWidgetListEntry>& B)
+			{
+				return bSortAscending ?
+					A->Count < B->Count :
+					A->Count > B->Count;
+			});
+	}
+	else if (SortColumn.IsEqual(SWidgetListRow::NAME_Type))
+	{
+		Entries.Sort([this](const TSharedPtr<FWidgetListEntry>& A, const TSharedPtr<FWidgetListEntry>& B)
+			{
+				return bSortAscending ?
+					A->Type.LexicalLess(B->Type) :
+					B->Type.LexicalLess(A->Type);
+			});
+	}
+	else if (SortColumn.IsEqual(SWidgetListRow::NAME_Location))
+	{
+		Entries.Sort([this](const TSharedPtr<FWidgetListEntry>& A, const TSharedPtr<FWidgetListEntry>& B)
+			{
+				return bSortAscending ?
+					A->Location.LexicalLess(B->Location) :
+					B->Location.LexicalLess(A->Location);
+			});
+	}
+
+	if (ListView.IsValid())
+	{
+		ListView->RequestListRefresh();
+	}
+}
+
+void SWidgetList::OnSort(const EColumnSortPriority::Type SortPriority, const FName& ColumnId, const EColumnSortMode::Type NewSortMode)
+{
+	if (SortPriority == EColumnSortPriority::Primary &&
+		NewSortMode != EColumnSortMode::None)
+	{
+		bSortAscending = NewSortMode == EColumnSortMode::Ascending;
+		SortColumn = ColumnId;
+
+		Sort();
+	}
+}
+
+TSharedRef<ITableRow> SWidgetList::OnGenerateRow(TSharedPtr<FWidgetListEntry> Entry, const TSharedRef<STableViewBase>& Table) const
+{
+	return SNew(SWidgetListRow, Entry, Table)
+		.OnAccessSource(OnAccessSource)
+		.OnAccessAsset(OnAccessAsset);
+}
+
+#undef LOCTEXT_NAMESPACE
diff --git a/Engine/Source/Developer/SlateReflector/Private/Widgets/SWidgetList.h b/Engine/Source/Developer/SlateReflector/Private/Widgets/SWidgetList.h
new file mode 100644
--- /dev/null
+++ b/Engine/Source/Developer/SlateReflector/Private/Widgets/SWidgetList.h
@@ -0,0 +1,41 @@
+// Copyright Epic Games, Inc. All Rights Reserved.
+
+#pragma once
+
+#include "AssetRegistry/AssetData.h"
+#include "Debugging/SlateDebugging.h"
+#include "Framework/SlateDelegates.h"
+#include "Widgets/DeclarativeSyntaxSupport.h"
+#include "Widgets/SCompoundWidget.h"
+#include "Widgets/Views/SListView.h"
+
+struct FWidgetListEntry;
+
+class SWidgetList : public SCompoundWidget
+{
+	SLATE_BEGIN_ARGS(SWidgetList) {}
+		SLATE_EVENT(FAccessAsset, OnAccessAsset)
+		SLATE_EVENT(FAccessSourceCode, OnAccessSource)
+	SLATE_END_ARGS()
+
+public:
+
+	void Construct(const FArguments& Args);
+
+private:
+	void Refresh();
+	void Sort();
+	void OnSort(const EColumnSortPriority::Type SortPriority, const FName& ColumnId, const EColumnSortMode::Type NewSortMode);
+	TSharedRef<ITableRow> OnGenerateRow(TSharedPtr<FWidgetListEntry> Entry, const TSharedRef<STableViewBase>& Table) const;
+	FReply ExportToCSV() const;
+
+private:
+	TSharedPtr<SListView<TSharedPtr<FWidgetListEntry>>> ListView;
+	TArray<TSharedPtr<FWidgetListEntry>> Entries;
+	FAccessAsset OnAccessAsset;
+	FAccessSourceCode OnAccessSource;
+	int32 CurrentCount = 0;
+	int32 PreviousCount = 0;
+	FName SortColumn;
+	bool bSortAscending = false;
+};
diff --git a/Engine/Source/Developer/SlateReflector/Private/Widgets/SWidgetReflector.cpp b/Engine/Source/Developer/SlateReflector/Private/Widgets/SWidgetReflector.cpp
--- a/Engine/Source/Developer/SlateReflector/Private/Widgets/SWidgetReflector.cpp
+++ b/Engine/Source/Developer/SlateReflector/Private/Widgets/SWidgetReflector.cpp
@@ -37,6 +37,7 @@
 #include "Widgets/SWidgetReflectorToolTipWidget.h"
 #include "Widgets/SWidgetEventLog.h"
 #include "Widgets/SWidgetHittestGrid.h"
+#include "Widgets/SWidgetList.h"
 #include "Widgets/SInvalidationPanel.h"
 #include "Widgets/Docking/SDockTab.h"
 #include "Widgets/Notifications/SNotificationList.h"
@@ -109,6 +110,7 @@ namespace WidgetReflectorTabID
 	static const FName SlateOptions = "WidgetReflector.SlateOptionsTab";
 	static const FName WidgetEvents = "WidgetReflector.WidgetEventsTab";
 	static const FName HittestGrid = "WidgetReflector.HittestGridTab";
+	static const FName WidgetList = "WidgetReflector.WidgetList";
 }
 
 namespace WidgetReflectorText
@@ -175,9 +177,13 @@ private:
 
 #if WITH_SLATE_DEBUGGING
 	TSharedRef<SDockTab> SpawnWidgetEvents(const FSpawnTabArgs& Args);
-	TSharedRef<SDockTab> SpawnWidgeHittestGrid(const FSpawnTabArgs& Args);
+	TSharedRef<SDockTab> SpawnWidgetHittestGrid(const FSpawnTabArgs& Args);
 #endif
 
+	TSharedRef<SDockTab> SpawnWidgetList(const FSpawnTabArgs& Args);
+
+public:
+
 	void HandleTabManagerPersistLayout(const TSharedRef<FTabManager::FLayout>& LayoutToSave);
 	void OnTabSpawned(const FName& TabIdentifier, const TSharedRef<SDockTab>& SpawnedTab);
 	void CloseTab(const FName& TabIdentifier);
@@ -533,6 +539,7 @@ void SWidgetReflector::Construct( const FArguments& InArgs )
 					->AddTab(WidgetReflectorTabID::WidgetEvents, ETabState::ClosedTab)
 					->AddTab(WidgetReflectorTabID::HittestGrid, ETabState::ClosedTab)
 #endif
+					->AddTab(WidgetReflectorTabID::WidgetList, ETabState::ClosedTab)
 				)
 			)
 #if WITH_EDITOR
@@ -581,10 +588,13 @@ void SWidgetReflector::Construct( const FArguments& InArgs )
 #if WITH_SLATE_DEBUGGING
 	RegisterTrackedTabSpawner(WidgetReflectorTabID::WidgetEvents, FOnSpawnTab::CreateSP(this, &SWidgetReflector::SpawnWidgetEvents))
 		.SetDisplayName(LOCTEXT("WidgetEventsTab", "Widget Events"));
-	RegisterTrackedTabSpawner(WidgetReflectorTabID::HittestGrid, FOnSpawnTab::CreateSP(this, &SWidgetReflector::SpawnWidgeHittestGrid))
+	RegisterTrackedTabSpawner(WidgetReflectorTabID::HittestGrid, FOnSpawnTab::CreateSP(this, &SWidgetReflector::SpawnWidgetHittestGrid))
 		.SetDisplayName(LOCTEXT("HitTestGridTab", "Hit Test Grid"));
 #endif
 
+	RegisterTrackedTabSpawner(WidgetReflectorTabID::WidgetList, FOnSpawnTab::CreateSP(this, &SWidgetReflector::SpawnWidgetList))
+		.SetDisplayName(LOCTEXT("WidgetListTab", "Widget List"));
+
 #if WITH_EDITOR
 	if (GIsEditor)
 	{
@@ -1017,7 +1027,7 @@ TSharedRef<SDockTab> SWidgetReflector::SpawnWidgetEvents(const FSpawnTabArgs& Ar
 		];
 }
 
-TSharedRef<SDockTab> SWidgetReflector::SpawnWidgeHittestGrid(const FSpawnTabArgs& Args)
+TSharedRef<SDockTab> SWidgetReflector::SpawnWidgetHittestGrid(const FSpawnTabArgs& Args)
 {
 	return SNew(SDockTab)
 		.Label(LOCTEXT("HitTestGridTab", "Hit Test Grid"))
@@ -1030,6 +1040,17 @@ TSharedRef<SDockTab> SWidgetReflector::SpawnWidgeHittestGrid(const FSpawnTabArgs
 
 #endif //WITH_SLATE_DEBUGGING
 
+TSharedRef<SDockTab> SWidgetReflector::SpawnWidgetList(const FSpawnTabArgs& Args)
+{
+	return SNew(SDockTab)
+		.Label(LOCTEXT("WidgetListTab", "Widget List"))
+		[
+			SNew(SWidgetList)
+			.OnAccessSource(SourceAccessDelegate)
+			.OnAccessAsset(AsseetAccessDelegate)
+		];
+}
+
 void SWidgetReflector::HandleTabManagerPersistLayout(const TSharedRef<FTabManager::FLayout>& LayoutToSave)
 {
 #if WITH_EDITOR
@@ -1040,7 +1061,6 @@ void SWidgetReflector::HandleTabManagerPersistLayout(const TSharedRef<FTabManage
 #endif //WITH_EDITOR
 }
 
-
 void SWidgetReflector::SaveSettings()
 {
 	GConfig->SetArray(TEXT("WidgetReflector"), TEXT("HiddenReflectorTreeColumns"), HiddenReflectorTreeColumns, *GEditorPerProjectIni);
diff --git a/Engine/Source/Developer/SlateReflector/SlateReflector.Build.cs b/Engine/Source/Developer/SlateReflector/SlateReflector.Build.cs
--- a/Engine/Source/Developer/SlateReflector/SlateReflector.Build.cs
+++ b/Engine/Source/Developer/SlateReflector/SlateReflector.Build.cs
@@ -22,7 +22,9 @@ public class SlateReflector : ModuleRules
 				"EditorStyle",
 				"Json",
 				"AssetRegistry",
-                "MessageLog"
+                "MessageLog",
+				"ToolWidgets",
+				"DesktopPlatform"
             }
 		);
 
diff --git a/Engine/Source/Runtime/SlateCore/Private/Debugging/SlateDebugging.cpp b/Engine/Source/Runtime/SlateCore/Private/Debugging/SlateDebugging.cpp
--- a/Engine/Source/Runtime/SlateCore/Private/Debugging/SlateDebugging.cpp
+++ b/Engine/Source/Runtime/SlateCore/Private/Debugging/SlateDebugging.cpp
@@ -2,17 +2,19 @@
 
 #include "Debugging/SlateDebugging.h"
 
-#include "SlateGlobals.h"
-#include "FastUpdate/WidgetProxy.h"
 #include "Animation/CurveSequence.h"
+#include "Application/SlateApplicationBase.h"
+#include "Debugging/WidgetList.h"
+#include "FastUpdate/WidgetProxy.h"
 #include "Layout/WidgetPath.h"
+#include "Misc/StringBuilder.h"
+#include "ProfilingDebugging/CsvProfiler.h"
 #include "Rendering/DrawElements.h"
+#include "SlateGlobals.h"
 #include "Styling/CoreStyle.h"
 #include "Types/ReflectionMetadata.h"
-#include "Widgets/SWidget.h"
 #include "Widgets/SNullWidget.h"
-#include "Application/SlateApplicationBase.h"
-#include "ProfilingDebugging/CsvProfiler.h"
+#include "Widgets/SWidget.h"
 
 #if UE_BUILD_SHIPPING
 CSV_DEFINE_CATEGORY_MODULE(SLATECORE_API, Slate, false);
@@ -662,6 +664,23 @@ void FSlateDebugging::BroadcastWidgetUpdatedByPaint(const SWidget* Invalidated,
 	}
 }
 
+const TArray<const SWidget*>& FSlateDebugging::GetAllWidgets()
+{
+#if UE_WITH_SLATE_DEBUG_WIDGETLIST
+	return UE::Slate::FWidgetList::GetAllWidgets();
+#else
+	static const TArray<const SWidget*> EmptyArray;
+	return EmptyArray;
+#endif
+}
+
+void FSlateDebugging::ExportWidgetList(FStringView Filename)
+{
+#if UE_WITH_SLATE_DEBUG_WIDGETLIST
+	UE::Slate::FWidgetList::ExportToCSV(*WriteToString<256>(Filename));
+#endif
+}
+
 #undef LOCTEXT_NAMESPACE
 
 #endif
\ No newline at end of file
diff --git a/Engine/Source/Runtime/SlateCore/Private/Debugging/WidgetList.cpp b/Engine/Source/Runtime/SlateCore/Private/Debugging/WidgetList.cpp
--- a/Engine/Source/Runtime/SlateCore/Private/Debugging/WidgetList.cpp
+++ b/Engine/Source/Runtime/SlateCore/Private/Debugging/WidgetList.cpp
@@ -320,6 +320,29 @@ FAutoConsoleCommandWithWorldArgsAndOutputDevice ConsoleCommandLogAllWidgets(
 	FConsoleCommandWithWorldArgsAndOutputDeviceDelegate::CreateStatic(&LogAllWidgetsDebugInfo)
 );
 
+void FWidgetList::ExportToCSV(FStringView OutputFilename)
+{
+	if (OutputFilename.Len() == 0)
+	{
+		return;
+	}
+
+	FLogAllWidgetsDebugInfoFlags Flags;
+	Flags.bDebug = true;
+	Flags.bPaint = true;
+	Flags.bProxy = true;
+	Flags.bChildren = true;
+	Flags.bParent = true;
+	Flags.bToolTip = true;
+	Flags.bCursor = true;
+	Flags.bSlateAttribute = true;
+	Flags.bMouseEventsHandler = true;
+
+	FOutputDeviceFile OutputDeviceFile(*WriteToString<256>(OutputFilename), true /*bDisableBackup*/);
+	OutputDeviceFile.SetSuppressEventTag(true);
+
+	LogAllWidgetsDebugInfoImpl(OutputDeviceFile, Flags);
+}
 
 } //Slate
 } //UE
diff --git a/Engine/Source/Runtime/SlateCore/Private/Debugging/WidgetList.h b/Engine/Source/Runtime/SlateCore/Private/Debugging/WidgetList.h
--- a/Engine/Source/Runtime/SlateCore/Private/Debugging/WidgetList.h
+++ b/Engine/Source/Runtime/SlateCore/Private/Debugging/WidgetList.h
@@ -6,13 +6,8 @@
 #include "Types/SlateAttribute.h"
 #include "Debugging/SlateDebugging.h"
 
-// Enabled to build a list of all the SWidget currently constructed
-#ifndef UE_WITH_SLATE_DEBUG_WIDGETLIST
-	#define UE_WITH_SLATE_DEBUG_WIDGETLIST 0 || UE_SLATE_WITH_MEMBER_ATTRIBUTE_DEBUGGING
-#endif
-
 #if UE_WITH_SLATE_DEBUG_WIDGETLIST && (UE_BUILD_SHIPPING || UE_BUILD_TEST)
-	#error "UE_WITH_SLATE_DEBUG_WIDGETLIST is defined in a none debug build"
+	#error "UE_WITH_SLATE_DEBUG_WIDGETLIST is defined in a non-debug build"
 #endif
 
 #if UE_WITH_SLATE_DEBUG_WIDGETLIST
@@ -32,6 +27,8 @@ public:
 
 	static const TArray<const SWidget*>& GetAllWidgets() { return AllWidgets; }
 
+	static void ExportToCSV(FStringView FileName);
+
 private:
 	FWidgetList() = delete;
 
diff --git a/Engine/Source/Runtime/SlateCore/Public/Debugging/SlateDebugging.h b/Engine/Source/Runtime/SlateCore/Public/Debugging/SlateDebugging.h
--- a/Engine/Source/Runtime/SlateCore/Public/Debugging/SlateDebugging.h
+++ b/Engine/Source/Runtime/SlateCore/Public/Debugging/SlateDebugging.h
@@ -7,6 +7,7 @@
 #include "FastUpdate/WidgetUpdateFlags.h"
 #include "Input/Reply.h"
 #include "ProfilingDebugging/CsvProfiler.h"
+#include "Types/SlateAttribute.h"
 #include "Widgets/InvalidateWidgetReason.h"
 
 #include "SlateDebugging.generated.h"
@@ -19,6 +20,11 @@
 	#define SLATE_CSV_TRACKER CSV_PROFILER
 #endif
 
+// Enabled to build a list of all the SWidget currently constructed
+#ifndef UE_WITH_SLATE_DEBUG_WIDGETLIST
+	#define UE_WITH_SLATE_DEBUG_WIDGETLIST UE_SLATE_WITH_MEMBER_ATTRIBUTE_DEBUGGING
+#endif
+
 class SWindow;
 class SWidget;
 struct FGeometry;
@@ -389,10 +395,13 @@ public:
 	DECLARE_MULTICAST_DELEGATE_TwoParams(FUICommandRun, const FName& /*CommandName*/, const FText& /*CommandLabel*/);
 	static FUICommandRun CommandRun;
 
+	static const TArray<const SWidget*>& GetAllWidgets();
+	static void ExportWidgetList(FStringView Filename);
+
 private:
 
 	// This class is only for namespace use
-	FSlateDebugging() = default;
+	FSlateDebugging() = delete;
 
 	struct FLastCursorQuery
 	{
