CIS fix for SWidgetList compile errors in shipping builds where WITH_SLATE_DEBUGGING is disabled.

diff --git a/Engine/Source/Developer/SlateReflector/Private/Widgets/SWidgetList.cpp b/Engine/Source/Developer/SlateReflector/Private/Widgets/SWidgetList.cpp
--- a/Engine/Source/Developer/SlateReflector/Private/Widgets/SWidgetList.cpp
+++ b/Engine/Source/Developer/SlateReflector/Private/Widgets/SWidgetList.cpp
@@ -217,6 +217,7 @@ void SWidgetList::Construct(const FArguments& Args)
 
 void SWidgetList::Refresh()
 {
+#if WITH_SLATE_DEBUGGING
 	TArray<const SWidget*> Widgets = FSlateDebugging::GetAllWidgets();
 
 	PreviousCount = CurrentCount;
@@ -254,10 +255,12 @@ void SWidgetList::Refresh()
 	}
 
 	Sort();
+#endif //WITH_SLATE_DEBUGGING
 }
 
 FReply SWidgetList::ExportToCSV() const
 {
+#if WITH_SLATE_DEBUGGING
 	FString OutputFile;
 	const TCHAR* DefaultFile = TEXT("WidgetList.csv");
 
@@ -293,6 +296,7 @@ FReply SWidgetList::ExportToCSV() const
 	FSlateDebugging::ExportWidgetList(OutputFile);
 
 	FPlatformProcess::ExploreFolder(*OutputFile);
+#endif //WITH_SLATE_DEBUGGING
 
 	return FReply::Handled();
 }
