Widget Inspector: Add the 'Source' name to the tooltip as the tooltip often covers up the source name.

diff --git a/Engine/Source/Developer/SlateReflector/Private/Widgets/SWidgetReflectorToolTipWidget.cpp b/Engine/Source/Developer/SlateReflector/Private/Widgets/SWidgetReflectorToolTipWidget.cpp
--- a/Engine/Source/Developer/SlateReflector/Private/Widgets/SWidgetReflectorToolTipWidget.cpp
+++ b/Engine/Source/Developer/SlateReflector/Private/Widgets/SWidgetReflectorToolTipWidget.cpp
@@ -69,6 +69,7 @@ void SReflectorToolTipWidget::Construct( const FArguments& InArgs )
 			++SlotCount;
 		};
 
+		BuildLabelAndValue(LOCTEXT("SourceName", "Source"), { this, &SReflectorToolTipWidget::GetWidgetsSourceName });
 		BuildLabelAndValue(LOCTEXT("DesiredSize", "Desired Size"), { this, &SReflectorToolTipWidget::GetWidgetsDesiredSize });
 		BuildLabelAndValue(LOCTEXT("ActualSize", "Actual Size"), { this, &SReflectorToolTipWidget::GetWidgetActualSize });
 		BuildLabelAndValue(LOCTEXT("SizeInfo", "Size Info"), { this, &SReflectorToolTipWidget::GetSizeInfo });
diff --git a/Engine/Source/Developer/SlateReflector/Private/Widgets/SWidgetReflectorToolTipWidget.h b/Engine/Source/Developer/SlateReflector/Private/Widgets/SWidgetReflectorToolTipWidget.h
--- a/Engine/Source/Developer/SlateReflector/Private/Widgets/SWidgetReflectorToolTipWidget.h
+++ b/Engine/Source/Developer/SlateReflector/Private/Widgets/SWidgetReflectorToolTipWidget.h
@@ -29,6 +29,11 @@ public:
 
 private:
 
+	FText GetWidgetsSourceName() const
+	{
+		return WidgetInfo->GetWidgetReadableLocation();
+	}
+
 	FText GetWidgetsDesiredSize() const
 	{
 		return FText::FromString(WidgetInfo->GetWidgetDesiredSize().ToString());
