WidgetReflector: Display the running validation test. Most of them only make sense in GlobalInvalidation. It's not that useful in that menu but it makes them easily discoverable.

diff --git a/Engine/Source/Developer/SlateReflector/Private/Widgets/SSlateOptions.cpp b/Engine/Source/Developer/SlateReflector/Private/Widgets/SSlateOptions.cpp
--- a/Engine/Source/Developer/SlateReflector/Private/Widgets/SSlateOptions.cpp
+++ b/Engine/Source/Developer/SlateReflector/Private/Widgets/SSlateOptions.cpp
@@ -24,7 +24,7 @@ void SSlateOptions::Construct( const FArguments& InArgs )
 {
 	struct Local
 	{
-		static void FillToolbar(FToolBarBuilder& ToolbarBuilder, const FSlateIcon& Icon, const FText& Label, const TCHAR* ConsoleVariable)
+		static void AddMenuEntry(FMenuBuilder& MenuBuilder, const FSlateIcon& Icon, const FText& Label, const TCHAR* ConsoleVariable, bool bCanEdit = true)
 		{
 			IConsoleVariable* CVar = IConsoleManager::Get().FindConsoleVariable(ConsoleVariable);
 			if (CVar)
@@ -33,45 +33,126 @@ void SSlateOptions::Construct( const FArguments& InArgs )
 				TooltipText.AppendLine(FString(CVar->GetHelp()));
 				TooltipText.AppendLine(FString(ConsoleVariable));
 
-				ToolbarBuilder.AddToolBarButton(
+
+				MenuBuilder.AddMenuEntry(
+					Label,
+					TooltipText.ToText(),
+					Icon,
 					FUIAction(
 						FExecuteAction::CreateLambda([CVar]() { CVar->Set(!CVar->GetBool(), EConsoleVariableFlags::ECVF_SetByCode); }),
-						FCanExecuteAction(),
+						FCanExecuteAction::CreateLambda([bCanEdit](){ return bCanEdit; }),
 						FGetActionCheckState::CreateLambda([CVar]() { return CVar->GetBool() ? ECheckBoxState::Checked : ECheckBoxState::Unchecked; })
 					),
 					NAME_None,
-					Label,
-					TooltipText.ToText(),
-					Icon,
-					EUserInterfaceActionType::ToggleButton
-				);
+					EUserInterfaceActionType::ToggleButton);
 			}
 		}
+	};
 
-		static void FillToolbar(FToolBarBuilder& ToolbarBuilder)
+	struct GlobalLocal : Local
+	{
+		static TSharedRef<SWidget> FillToolbar()
 		{
+			const bool bShouldCloseWindowAfterMenuSelection = true;
+			FMenuBuilder MenuBuilder(bShouldCloseWindowAfterMenuSelection, nullptr);
+
 			FSlateIcon Icon(FWidgetReflectorStyle::GetStyleSetName(), "Icon.Empty");
 
-			ToolbarBuilder.BeginSection("Flags");
-			{
+			AddMenuEntry(MenuBuilder, Icon, LOCTEXT("EnableWidgetCaching", "Fast Widget Path"), TEXT("Slate.EnableFastWidgetPath"), false);
+			AddMenuEntry(MenuBuilder, Icon, LOCTEXT("EnableToolTips", "Enable Tooltips"), TEXT("Slate.EnableTooltips"));
+			AddMenuEntry(MenuBuilder, Icon, LOCTEXT("GlobalInvalidation", "Global Invalidation"), TEXT("Slate.EnableGlobalInvalidation"));
+
+			return MenuBuilder.MakeWidget();
+		}
+	};
+
+	struct DebugLocal : Local
+	{
+		static TSharedRef<SWidget> FillToolbar()
+		{
+			const bool bShouldCloseWindowAfterMenuSelection = true;
+			FMenuBuilder MenuBuilder(bShouldCloseWindowAfterMenuSelection, nullptr);
+
+			FSlateIcon Icon(FWidgetReflectorStyle::GetStyleSetName(), "Icon.Empty");
 #if WITH_SLATE_DEBUGGING
-				FillToolbar(ToolbarBuilder, Icon, LOCTEXT("EnableWidgetCaching", "Widget Caching"), TEXT("Slate.EnableInvalidationPanels"));
-				FillToolbar(ToolbarBuilder, Icon, LOCTEXT("InvalidationDebugging", "Invalidation Debugging"), TEXT("SlateDebugger.Invalidate.Enable"));
-				FillToolbar(ToolbarBuilder, Icon, LOCTEXT("InvalidationRootDebugging", "Invalidation Root Debugging"), TEXT("SlateDebugger.InvalidationRoot.Enable"));
-				FillToolbar(ToolbarBuilder, Icon, LOCTEXT("UpdateDebugging", "Update Debugging"), TEXT("SlateDebugger.Update.Enable"));
-				FillToolbar(ToolbarBuilder, Icon, LOCTEXT("PaintDebugging", "Paint Debugging"), TEXT("SlateDebugger.Paint.Enable"));
-				FillToolbar(ToolbarBuilder, Icon, LOCTEXT("ShowClipping", "Show Clipping"), TEXT("Slate.ShowClipping"));
-				FillToolbar(ToolbarBuilder, Icon, LOCTEXT("DebugCulling", "Debug Culling"), TEXT("Slate.DebugCulling"));
-				FillToolbar(ToolbarBuilder, Icon, LOCTEXT("EnsureAllVisibleWidgetsPaint", "Ensure All Visible Widgets Paint"), TEXT("Slate.EnsureAllVisibleWidgetsPaint"));
+			AddMenuEntry(MenuBuilder, Icon, LOCTEXT("EnableWidgetCaching", "Enable InvalidationBox"), TEXT("Slate.EnableInvalidationPanels"));
+			AddMenuEntry(MenuBuilder, Icon, LOCTEXT("InvalidationDebugging", "Show Invalidation"), TEXT("SlateDebugger.Invalidate.Enable"));
+			AddMenuEntry(MenuBuilder, Icon, LOCTEXT("InvalidationRootDebugging", "Show Root Invalidation"), TEXT("SlateDebugger.InvalidationRoot.Enable"));
+			AddMenuEntry(MenuBuilder, Icon, LOCTEXT("UpdateDebugging", "Show Update"), TEXT("SlateDebugger.Update.Enable"));
+			AddMenuEntry(MenuBuilder, Icon, LOCTEXT("PaintDebugging", "Show Paint"), TEXT("SlateDebugger.Paint.Enable"));
+			AddMenuEntry(MenuBuilder, Icon, LOCTEXT("ShowClipping", "Show Clipping"), TEXT("Slate.ShowClipping"));
+			AddMenuEntry(MenuBuilder, Icon, LOCTEXT("DebugCulling", "Debug Culling"), TEXT("Slate.DebugCulling"));
+			AddMenuEntry(MenuBuilder, Icon, LOCTEXT("DebugCulling", "Show HitTestGrid"), TEXT("Slate.HitTestGridDebugging"));
 #endif // WITH_SLATE_DEBUGGING
-			}
-			ToolbarBuilder.EndSection();
+			return MenuBuilder.MakeWidget();
 		}
 	};
 
+	struct ValidationLocal : Local
+	{
+		static TSharedRef<SWidget> FillToolbar()
+		{
+			const bool bShouldCloseWindowAfterMenuSelection = true;
+			FMenuBuilder MenuBuilder(bShouldCloseWindowAfterMenuSelection, nullptr);
+
+			FSlateIcon Icon(FWidgetReflectorStyle::GetStyleSetName(), "Icon.Empty");
+#if WITH_SLATE_DEBUGGING
+			AddMenuEntry(MenuBuilder, Icon, LOCTEXT("VerifyParentChildrenRelationship", "Verify Parent/Children Relationship"), TEXT("Slate.VerifyParentChildrenRelationship"));
+			AddMenuEntry(MenuBuilder, Icon, LOCTEXT("VerifyLayerId", "Verify LayerId"), TEXT("Slate.VerifyWidgetLayerId"));
+#endif // WITH_SLATE_DEBUGGING
+
+			{
+				MenuBuilder.BeginSection("InvalidationRoot", LOCTEXT("Invalidation", "Invalidation"));
+	#if WITH_SLATE_DEBUGGING
+				AddMenuEntry(MenuBuilder, Icon, LOCTEXT("EnsureAllVisibleWidgetsPaint", "Verify Visible Widgets Paint"), TEXT("Slate.EnsureAllVisibleWidgetsPaint"));
+	#endif // WITH_SLATE_DEBUGGING
+	#if UE_SLATE_WITH_INVALIDATIONWIDGETLIST_DEBUGGING
+				AddMenuEntry(MenuBuilder, Icon, LOCTEXT("VerifyWidgetList", "Verify Widget List"), TEXT("Slate.InvalidationRoot.VerifyWidgetList"));
+				AddMenuEntry(MenuBuilder, Icon, LOCTEXT("VerifyWidgetIndex", "Verify Widget's index"), TEXT("Slate.InvalidationRoot.VerifyWidgetsIndex"));
+				AddMenuEntry(MenuBuilder, Icon, LOCTEXT("VerifyWidgetPtr", "Verify Widget Pointer"), TEXT("Slate.InvalidationRoot.VerifyValidWidgets"));
+				AddMenuEntry(MenuBuilder, Icon, LOCTEXT("VerifyHittestGrid", "Verify Hittest Grid"), TEXT("Slate.InvalidationRoot.VerifyHittestGrid"));
+				AddMenuEntry(MenuBuilder, Icon, LOCTEXT("VefiryVisibility", "Verify Visibility"), TEXT("Slate.InvalidationRoot.VerifyWidgetVisibility"));
+				AddMenuEntry(MenuBuilder, Icon, LOCTEXT("VerifyVolatility", "Verify Volatility"), TEXT("Slate.InvalidationRoot.VerifyWidgetVolatile"));
+				AddMenuEntry(MenuBuilder, Icon, LOCTEXT("VerifyUpdateList", "Verify Update List"), TEXT("Slate.InvalidationRoot.VerifyWidgetUpdateList"));
+				AddMenuEntry(MenuBuilder, Icon, LOCTEXT("VerifySlateAttributes", "Verify Attributes"), TEXT("Slate.InvalidationRoot.VerifySlateAttribute"));
+	#endif // UE_SLATE_WITH_INVALIDATIONWIDGETLIST_DEBUGGING
+				MenuBuilder.EndSection();
+			}			
+
+			return MenuBuilder.MakeWidget();
+		}
+	};
+
+
 	FToolBarBuilder ToolbarBuilder(TSharedPtr<const FUICommandList>(), FMultiBoxCustomization::None);
 	ToolbarBuilder.SetStyle(&FWidgetReflectorStyle::Get(), "BoldSlimToolbar");
-	Local::FillToolbar(ToolbarBuilder);
+
+	ToolbarBuilder.AddComboButton(
+		FUIAction(),
+		FOnGetContent::CreateStatic(&GlobalLocal::FillToolbar),
+		LOCTEXT("FlagLabel", "Flags"),
+		FText::GetEmpty(),
+		FSlateIcon(),
+		false
+	);
+
+	ToolbarBuilder.AddComboButton(
+		FUIAction(),
+		FOnGetContent::CreateStatic(&DebugLocal::FillToolbar),
+		LOCTEXT("DebugLabel", "Debug Options"),
+		FText::GetEmpty(),
+		FSlateIcon(),
+		false
+	);
+
+	ToolbarBuilder.AddComboButton(
+		FUIAction(),
+		FOnGetContent::CreateStatic(&ValidationLocal::FillToolbar),
+		LOCTEXT("EnsureLabel", "Runtime Validation"),
+		FText::GetEmpty(),
+		FSlateIcon(),
+		false
+	);
 	
 	ChildSlot
 	[
@@ -106,15 +187,6 @@ void SSlateOptions::Construct( const FArguments& InArgs )
 				]
 			]
 
-			+ SHorizontalBox::Slot()
-			.AutoWidth()
-			.VAlign(VAlign_Center)
-			.Padding(20.f, 0.f, 4.f, 0.f)
-			[
-				SNew(STextBlock)
-				.Text(LOCTEXT("Flags", "Flags: "))
-			]
-
 			+ SHorizontalBox::Slot()
 			.FillWidth(1)
 			.Padding(FMargin(5.0f, 0.0f))
diff --git a/Engine/Source/Developer/SlateReflector/Private/Widgets/SWidgetReflector.cpp b/Engine/Source/Developer/SlateReflector/Private/Widgets/SWidgetReflector.cpp
--- a/Engine/Source/Developer/SlateReflector/Private/Widgets/SWidgetReflector.cpp
+++ b/Engine/Source/Developer/SlateReflector/Private/Widgets/SWidgetReflector.cpp
@@ -569,7 +569,7 @@ void SWidgetReflector::Construct( const FArguments& InArgs )
 	TabManager->SetOnPersistLayout(FTabManager::FOnPersistLayout::CreateRaw(this, &SWidgetReflector::HandleTabManagerPersistLayout));
 
 	RegisterTrackedTabSpawner(WidgetReflectorTabID::SlateOptions, FOnSpawnTab::CreateSP(this, &SWidgetReflector::SpawnSlateOptionWidgetTab))
-		.SetDisplayName(LOCTEXT("OptionsTab", "Slate Debug Options"));
+		.SetDisplayName(LOCTEXT("OptionsTab", "Toolbar"));
 
 	RegisterTrackedTabSpawner(WidgetReflectorTabID::WidgetHierarchy, FOnSpawnTab::CreateSP(this, &SWidgetReflector::SpawnWidgetHierarchyTab))
 		.SetDisplayName(LOCTEXT("WidgetHierarchyTab", "Widget Hierarchy"));
@@ -679,7 +679,7 @@ void SWidgetReflector::HandlePullDownWindowMenu(FMenuBuilder& MenuBuilder)
 TSharedRef<SDockTab> SWidgetReflector::SpawnSlateOptionWidgetTab(const FSpawnTabArgs& Args)
 {
 	TSharedRef<SDockTab> SpawnedTab = SNew(SDockTab)
-		.Label(LOCTEXT("WidgetHierarchyTab", "Widget Hierarchy"))
+		.Label(LOCTEXT("ToolbarTab", "Toolbar"))
 		.ShouldAutosize(true)
 		[
 			SNew(SSlateOptions)
