Dismiss tooltip from widget reflector when context menu is spawning.

diff --git a/Engine/Source/Developer/SlateReflector/Private/Widgets/SWidgetReflector.cpp b/Engine/Source/Developer/SlateReflector/Private/Widgets/SWidgetReflector.cpp
--- a/Engine/Source/Developer/SlateReflector/Private/Widgets/SWidgetReflector.cpp
+++ b/Engine/Source/Developer/SlateReflector/Private/Widgets/SWidgetReflector.cpp
@@ -2033,6 +2033,9 @@ void SWidgetReflector::HandleReflectorTreeSelectionChanged( TSharedPtr<FWidgetRe
 
 TSharedRef<SWidget> SWidgetReflector::HandleReflectorTreeContextMenu()
 {
+	// We spawn a large tooltip, close it immediately to prevent context menu from hiding.
+	FSlateApplication::Get().CloseToolTip();
+
 	const bool bShouldCloseWindowAfterMenuSelection = true;
 	FMenuBuilder MenuBuilder(bShouldCloseWindowAfterMenuSelection, nullptr);
 
