Slate: Fix duplicated localization key.
#rnx

diff --git a/Engine/Source/Developer/SlateReflector/Private/Widgets/SSlateOptions.cpp b/Engine/Source/Developer/SlateReflector/Private/Widgets/SSlateOptions.cpp
--- a/Engine/Source/Developer/SlateReflector/Private/Widgets/SSlateOptions.cpp
+++ b/Engine/Source/Developer/SlateReflector/Private/Widgets/SSlateOptions.cpp
@@ -58,7 +58,7 @@ void SSlateOptions::Construct( const FArguments& InArgs )
 
 			FSlateIcon Icon(FWidgetReflectorStyle::GetStyleSetName(), "Icon.Empty");
 
-			AddMenuEntry(MenuBuilder, Icon, LOCTEXT("EnableWidgetCaching", "Fast Widget Path"), TEXT("Slate.EnableFastWidgetPath"), false);
+			AddMenuEntry(MenuBuilder, Icon, LOCTEXT("EnableFastWidgetPath", "Fast Widget Path"), TEXT("Slate.EnableFastWidgetPath"), false);
 			AddMenuEntry(MenuBuilder, Icon, LOCTEXT("EnableToolTips", "Enable Tooltips"), TEXT("Slate.EnableTooltips"));
 			AddMenuEntry(MenuBuilder, Icon, LOCTEXT("GlobalInvalidation", "Global Invalidation"), TEXT("Slate.EnableGlobalInvalidation"));
 
@@ -75,7 +75,7 @@ void SSlateOptions::Construct( const FArguments& InArgs )
 
 			FSlateIcon Icon(FWidgetReflectorStyle::GetStyleSetName(), "Icon.Empty");
 #if WITH_SLATE_DEBUGGING
-			AddMenuEntry(MenuBuilder, Icon, LOCTEXT("EnableWidgetCaching", "Enable InvalidationBox"), TEXT("Slate.EnableInvalidationPanels"));
+			AddMenuEntry(MenuBuilder, Icon, LOCTEXT("EnableInvalidationPanels", "Enable InvalidationBox"), TEXT("Slate.EnableInvalidationPanels"));
 			AddMenuEntry(MenuBuilder, Icon, LOCTEXT("InvalidationDebugging", "Show Invalidation"), TEXT("SlateDebugger.Invalidate.Enable"));
 			AddMenuEntry(MenuBuilder, Icon, LOCTEXT("InvalidationRootDebugging", "Show Root Invalidation"), TEXT("SlateDebugger.InvalidationRoot.Enable"));
 			AddMenuEntry(MenuBuilder, Icon, LOCTEXT("UpdateDebugging", "Show Update"), TEXT("SlateDebugger.Update.Enable"));
