Fix duplicated localization key
#rnx

diff --git a/Engine/Source/Developer/SlateReflector/Private/Widgets/SSlateOptions.cpp b/Engine/Source/Developer/SlateReflector/Private/Widgets/SSlateOptions.cpp
--- a/Engine/Source/Developer/SlateReflector/Private/Widgets/SSlateOptions.cpp
+++ b/Engine/Source/Developer/SlateReflector/Private/Widgets/SSlateOptions.cpp
@@ -82,7 +82,7 @@ void SSlateOptions::Construct( const FArguments& InArgs )
 			AddMenuEntry(MenuBuilder, Icon, LOCTEXT("PaintDebugging", "Show Paint"), TEXT("SlateDebugger.Paint.Enable"));
 			AddMenuEntry(MenuBuilder, Icon, LOCTEXT("ShowClipping", "Show Clipping"), TEXT("Slate.ShowClipping"));
 			AddMenuEntry(MenuBuilder, Icon, LOCTEXT("DebugCulling", "Debug Culling"), TEXT("Slate.DebugCulling"));
-			AddMenuEntry(MenuBuilder, Icon, LOCTEXT("DebugCulling", "Show HitTestGrid"), TEXT("Slate.HitTestGridDebugging"));
+			AddMenuEntry(MenuBuilder, Icon, LOCTEXT("ShowHitTestGrid", "Show HitTestGrid"), TEXT("Slate.HitTestGridDebugging"));
 #endif // WITH_SLATE_DEBUGGING
 			return MenuBuilder.MakeWidget();
 		}
