Widget reflector menu is not shown on OSX
#fix ensure that MenuBar widget is applied to respective TabManager

diff --git a/Engine/Source/Developer/SlateReflector/Private/Widgets/SWidgetReflector.cpp b/Engine/Source/Developer/SlateReflector/Private/Widgets/SWidgetReflector.cpp
--- a/Engine/Source/Developer/SlateReflector/Private/Widgets/SWidgetReflector.cpp
+++ b/Engine/Source/Developer/SlateReflector/Private/Widgets/SWidgetReflector.cpp
@@ -624,6 +624,9 @@ void SWidgetReflector::Construct( const FArguments& InArgs )
 		"Window"
 		);
 
+	const TSharedRef<SWidget> MenuWidget = MenuBarBuilder.MakeWidget();
+	TabManager->SetMenuMultiBox(MenuBarBuilder.GetMultiBox(), MenuWidget);
+
 	this->ChildSlot
 	[
 		SNew(SBorder)
@@ -636,7 +639,7 @@ void SWidgetReflector::Construct( const FArguments& InArgs )
 			.AutoHeight()
 			.Padding(FMargin(0.0f, 4.0f, 0.0f, 0.0f))
 			[
-				MenuBarBuilder.MakeWidget()
+				MenuWidget
 			]
 			
 			+ SVerticalBox::Slot()
