SlateReflector: Add the return value from OnPaint to the WidgetReflector.

diff --git a/Engine/Source/Developer/SlateReflector/Private/Models/WidgetReflectorNode.cpp b/Engine/Source/Developer/SlateReflector/Private/Models/WidgetReflectorNode.cpp
--- a/Engine/Source/Developer/SlateReflector/Private/Models/WidgetReflectorNode.cpp
+++ b/Engine/Source/Developer/SlateReflector/Private/Models/WidgetReflectorNode.cpp
@@ -149,6 +149,10 @@ int32 FLiveWidgetReflectorNode::GetWidgetLayerId() const
 {
 	return FWidgetReflectorNodeUtils::GetWidgetLayerId(Widget.Pin());
 }
+int32 FLiveWidgetReflectorNode::GetWidgetLayerIdOut() const
+{
+	return FWidgetReflectorNodeUtils::GetWidgetLayerIdOut(Widget.Pin());
+}
 
 bool FLiveWidgetReflectorNode::GetWidgetFocusable() const
 {
@@ -275,6 +279,7 @@ FSnapshotWidgetReflectorNode::FSnapshotWidgetReflectorNode(const FArrangedWidget
 	, bCachedWidgetEnabled(FWidgetReflectorNodeUtils::GetWidgetEnabled(InArrangedWidget.Widget))
 	, CachedWidgetClippingText(FWidgetReflectorNodeUtils::GetWidgetClippingText(InArrangedWidget.Widget))
 	, CachedWidgetLayerId(FWidgetReflectorNodeUtils::GetWidgetLayerId(InArrangedWidget.Widget))
+	, CachedWidgetLayerIdOut(FWidgetReflectorNodeUtils::GetWidgetLayerIdOut(InArrangedWidget.Widget))
 	, CachedWidgetReadableLocation(FWidgetReflectorNodeUtils::GetWidgetReadableLocation(InArrangedWidget.Widget))
 	, CachedWidgetFile(FWidgetReflectorNodeUtils::GetWidgetFile(InArrangedWidget.Widget))
 	, CachedWidgetLineNumber(FWidgetReflectorNodeUtils::GetWidgetLineNumber(InArrangedWidget.Widget))
@@ -337,6 +342,11 @@ int32 FSnapshotWidgetReflectorNode::GetWidgetLayerId() const
 	return CachedWidgetLayerId;
 }
 
+int32 FSnapshotWidgetReflectorNode::GetWidgetLayerIdOut() const
+{
+	return CachedWidgetLayerIdOut;
+}
+
 bool FSnapshotWidgetReflectorNode::GetWidgetNeedsTick() const
 {
 	return bCachedWidgetNeedsTick;
@@ -518,6 +528,7 @@ TSharedRef<FJsonValue> FSnapshotWidgetReflectorNode::ToJson(const TSharedRef<FSn
 	RootJsonObject->SetBoolField(TEXT("WidgetEnabled"), RootSnapshotNode->bCachedWidgetEnabled);
 	RootJsonObject->SetStringField(TEXT("WidgetClippingText"), RootSnapshotNode->CachedWidgetClippingText.ToString());
 	RootJsonObject->SetNumberField(TEXT("WidgetLayerId"), RootSnapshotNode->CachedWidgetLayerId);
+	RootJsonObject->SetNumberField(TEXT("WidgetLayerIdOut"), RootSnapshotNode->CachedWidgetLayerIdOut);
 	RootJsonObject->SetStringField(TEXT("WidgetReadableLocation"), RootSnapshotNode->CachedWidgetReadableLocation.ToString());
 	RootJsonObject->SetStringField(TEXT("WidgetFile"), RootSnapshotNode->CachedWidgetFile);
 	RootJsonObject->SetNumberField(TEXT("WidgetLineNumber"), RootSnapshotNode->CachedWidgetLineNumber);
@@ -691,6 +702,7 @@ TSharedRef<FSnapshotWidgetReflectorNode> FSnapshotWidgetReflectorNode::FromJson(
 	RootSnapshotNode->bCachedWidgetEnabled = RootJsonObject->GetBoolField(TEXT("WidgetEnabled"));
 	RootSnapshotNode->CachedWidgetClippingText = FText::FromString(RootJsonObject->GetStringField(TEXT("WidgetClippingText")));
 	RootSnapshotNode->CachedWidgetLayerId = RootJsonObject->GetIntegerField(TEXT("WidgetLayerId"));
+	RootSnapshotNode->CachedWidgetLayerIdOut = RootJsonObject->GetIntegerField(TEXT("WidgetLayerIdOut"));
 	RootSnapshotNode->CachedWidgetReadableLocation = FText::FromString(RootJsonObject->GetStringField(TEXT("WidgetReadableLocation")));
 	RootSnapshotNode->CachedWidgetFile = RootJsonObject->GetStringField(TEXT("WidgetFile"));
 	RootSnapshotNode->CachedWidgetLineNumber = RootJsonObject->GetIntegerField(TEXT("WidgetLineNumber"));
@@ -1017,6 +1029,11 @@ int32 FWidgetReflectorNodeUtils::GetWidgetLayerId(const TSharedPtr<const SWidget
 	return (InWidget.IsValid()) ? InWidget->GetPersistentState().LayerId : -1;
 }
 
+int32 FWidgetReflectorNodeUtils::GetWidgetLayerIdOut(const TSharedPtr<const SWidget>& InWidget)
+{
+	return (InWidget.IsValid()) ? InWidget->GetPersistentState().OutgoingLayerId : -1;
+}
+
 FText FWidgetReflectorNodeUtils::GetWidgetReadableLocation(const TSharedPtr<const SWidget>& InWidget)
 {
 	return FText::FromString(FReflectionMetaData::GetWidgetDebugInfo(InWidget.Get()));
diff --git a/Engine/Source/Developer/SlateReflector/Private/Models/WidgetReflectorNode.h b/Engine/Source/Developer/SlateReflector/Private/Models/WidgetReflectorNode.h
--- a/Engine/Source/Developer/SlateReflector/Private/Models/WidgetReflectorNode.h
+++ b/Engine/Source/Developer/SlateReflector/Private/Models/WidgetReflectorNode.h
@@ -80,6 +80,7 @@ public:
 	 * @return The LayerId the widget were painted from
 	 */
 	virtual int32 GetWidgetLayerId() const = 0;
+	virtual int32 GetWidgetLayerIdOut() const = 0;
 
 	/**
 	 * @return The bool indicating whether or not the widget we were initialized from reports as Focusable
@@ -287,6 +288,7 @@ public:
 	virtual bool GetWidgetVisibilityInherited() const override;
 	virtual FText GetWidgetClippingText() const override;
 	virtual int32 GetWidgetLayerId() const override;
+	virtual int32 GetWidgetLayerIdOut() const override;
 	virtual bool GetWidgetFocusable() const override;
 	virtual bool GetWidgetNeedsTick() const override;
 	virtual bool GetWidgetIsVolatile() const override;
@@ -346,6 +348,7 @@ public:
 	virtual FText GetWidgetVisibilityText() const override;
 	virtual FText GetWidgetClippingText() const override;
 	virtual int32 GetWidgetLayerId() const override;
+	virtual int32 GetWidgetLayerIdOut() const override;
 	virtual bool GetWidgetVisible() const override;
 	virtual bool GetWidgetVisibilityInherited() const override;
 	virtual bool GetWidgetFocusable() const override;
@@ -426,6 +429,7 @@ private:
 
 	/** The LayerId of the widget */
 	int32 CachedWidgetLayerId;
+	int32 CachedWidgetLayerIdOut;
 
 	/** The human readable location (source file for C++ widgets, asset name for UMG widgets) of the widget at the point it was passed to Initialize */
 	FText CachedWidgetReadableLocation;
@@ -573,6 +577,7 @@ public:
 	 * @return The current LayerId for the given widget
 	 */
 	static int32 GetWidgetLayerId(const TSharedPtr<const SWidget>& InWidget);
+	static int32 GetWidgetLayerIdOut(const TSharedPtr<const SWidget>& InWidget);
 
 	/**
 	* @return The current focusability for the given widget
diff --git a/Engine/Source/Developer/SlateReflector/Private/Widgets/SWidgetReflectorTreeWidgetItem.cpp b/Engine/Source/Developer/SlateReflector/Private/Widgets/SWidgetReflectorTreeWidgetItem.cpp
--- a/Engine/Source/Developer/SlateReflector/Private/Widgets/SWidgetReflectorTreeWidgetItem.cpp
+++ b/Engine/Source/Developer/SlateReflector/Private/Widgets/SWidgetReflectorTreeWidgetItem.cpp
@@ -149,7 +149,7 @@ TSharedRef<SWidget> SReflectorTreeWidgetItem::GenerateWidgetForColumn(const FNam
 			.Padding(FMargin(2.0f, 0.0f))
 			[
 				SNew(STextBlock)
-				.Text(FText::AsNumber(WidgetInfo->GetWidgetLayerId()))
+				.Text(FText::Format(LOCTEXT("WidgetLayerIds", "[{0}, {1}]"), FText::AsNumber(WidgetInfo->GetWidgetLayerId()), FText::AsNumber(WidgetInfo->GetWidgetLayerIdOut())))
 			];
 	}
 	else if (ColumnName == NAME_ForegroundColor )
