FNames containing asset paths are deprecated. FSoftObjectPath or FTopLevelAssetPath should be used instead.
Fixups for misc developer dir files.

diff --git a/Engine/Plugins/Developer/Concert/ConcertSync/ConcertSyncClient/Source/ConcertSyncClient/Private/ConcertSyncClientUtil.cpp b/Engine/Plugins/Developer/Concert/ConcertSync/ConcertSyncClient/Source/ConcertSyncClient/Private/ConcertSyncClientUtil.cpp
--- a/Engine/Plugins/Developer/Concert/ConcertSync/ConcertSyncClient/Source/ConcertSyncClient/Private/ConcertSyncClientUtil.cpp
+++ b/Engine/Plugins/Developer/Concert/ConcertSync/ConcertSyncClient/Source/ConcertSyncClient/Private/ConcertSyncClientUtil.cpp
@@ -252,7 +252,7 @@ FGetObjectResult GetObject(const FConcertObjectId& InObjectId, const FName InNew
 	// TODO: If a case arises where we need to go the other direction and get
 	// an object with a non-partitioned path from a partitioned path, a
 	// different mechanism for that would be needed here.
-	if (UObject* ExistingObjectOuter = FSoftObjectPath(ObjectOuterPathToFind).ResolveObject())
+	if (UObject* ExistingObjectOuter = FSoftObjectPath(ObjectOuterPathToFind.ToString()).ResolveObject())
 	{
 		// We need the object class to find or create the object
 		if (UClass* ObjectClass = FindOrLoadClass(InObjectId.ObjectClassPathName))
@@ -302,7 +302,7 @@ FGetObjectResult GetObject(const FConcertObjectId& InObjectId, const FName InNew
 	// Find the outer for the new object.
 	// As above, we use FSoftObjectPath::ResolveObject() here to account for
 	// the possibility of world partitioning.
-	if (UObject* NewObjectOuter = FSoftObjectPath(ObjectOuterPathToCreate).ResolveObject())
+	if (UObject* NewObjectOuter = FSoftObjectPath(ObjectOuterPathToCreate.ToString()).ResolveObject())
 	{
 		// We need the object class to find or create the object
 		if (UClass* ObjectClass = FindOrLoadClass(InObjectId.ObjectClassPathName))
diff --git a/Engine/Plugins/Developer/TraceSourceFiltering/Source/SourceFilteringTrace/Private/SourceFilterPresets.cpp b/Engine/Plugins/Developer/TraceSourceFiltering/Source/SourceFilteringTrace/Private/SourceFilterPresets.cpp
--- a/Engine/Plugins/Developer/TraceSourceFiltering/Source/SourceFilteringTrace/Private/SourceFilterPresets.cpp
+++ b/Engine/Plugins/Developer/TraceSourceFiltering/Source/SourceFilteringTrace/Private/SourceFilterPresets.cpp
@@ -80,7 +80,7 @@ void FSourceFilterPresets::LoadPresetCommand(const TArray<FString>& Arguments)
 			// Retrieve specific asset data at user index
 			if (PresetAssetData.IsValidIndex(Index))
 			{
-				LoadPreset(PresetAssetData[Index].ObjectPath.ToString());
+				LoadPreset(PresetAssetData[Index].GetObjectPathString());
 			}
 			else
 			{
diff --git a/Engine/Source/Developer/FunctionalTesting/Private/FunctionalTestingModule.cpp b/Engine/Source/Developer/FunctionalTesting/Private/FunctionalTestingModule.cpp
--- a/Engine/Source/Developer/FunctionalTesting/Private/FunctionalTestingModule.cpp
+++ b/Engine/Source/Developer/FunctionalTesting/Private/FunctionalTestingModule.cpp
@@ -125,7 +125,7 @@ void FFunctionalTestingModule::GetMapTests(bool bEditorOnlyTests, TArray<FString
 
 			for (const FAssetData& MapAsset : MapList)
 			{
-				FString MapAssetPath = MapAsset.ObjectPath.ToString();
+				FString MapAssetPath = MapAsset.GetObjectPathString();
 				FString MapPackageName = MapAsset.PackageName.ToString();
 				if (!IsDeveloperDirectoryIncluded && MapPackageName.Find(TEXT("/Game/Developers")) == 0) continue;
 				FString PartialSuiteName = MapPackageName.RightChop(6).Replace(TEXT("/"), TEXT(".")); // Remove "/Game/" from the name and use dot syntax
diff --git a/Engine/Source/Developer/IoStoreUtilities/Private/IoStoreUtilities.cpp b/Engine/Source/Developer/IoStoreUtilities/Private/IoStoreUtilities.cpp
--- a/Engine/Source/Developer/IoStoreUtilities/Private/IoStoreUtilities.cpp
+++ b/Engine/Source/Developer/IoStoreUtilities/Private/IoStoreUtilities.cpp
@@ -3052,7 +3052,7 @@ static void AddChunkInfoToAssetRegistry(TMap<FPackageId, TArray<FIoStoreTocChunk
 		TagsAndValues.Add("Stage_ChunkCount", LexToString(ChunkCount));
 		TagsAndValues.Add("Stage_ChunkSize", LexToString(Size));
 		TagsAndValues.Add("Stage_ChunkCompressedSize", LexToString(CompressedSize));
-		AssetRegistry.AddTagsToAssetData(FSoftObjectPath(AssetData->ObjectPath), MoveTemp(TagsAndValues));
+		AssetRegistry.AddTagsToAssetData(AssetData->GetSoftObjectPath(), MoveTemp(TagsAndValues));
 
 		// We assign a package's chunks to a single asset, remove it from the list so that
 		// at the end we can track how many chunks don't get assigned.
diff --git a/Engine/Source/Developer/RigVMDeveloper/Private/RigVMModel/RigVMController.cpp b/Engine/Source/Developer/RigVMDeveloper/Private/RigVMModel/RigVMController.cpp
--- a/Engine/Source/Developer/RigVMDeveloper/Private/RigVMModel/RigVMController.cpp
+++ b/Engine/Source/Developer/RigVMDeveloper/Private/RigVMModel/RigVMController.cpp
@@ -12098,14 +12098,14 @@ TArray<FAssetData> URigVMController::GetAffectedAssets(ERigVMControllerBulkEditT
 
 	URigVMGraph* Graph = GetGraph();
 	TSoftObjectPtr<URigVMGraph> GraphPtr = Graph;
-	const FString ThisAssetPath = GraphPtr.ToSoftObjectPath().GetAssetPathName().ToString();
+	const FString ThisAssetPath = GraphPtr.ToSoftObjectPath().GetAssetPath().ToString();
 
 	const FAssetRegistryModule& AssetRegistryModule = FModuleManager::LoadModuleChecked<FAssetRegistryModule>(TEXT("AssetRegistry"));
 
 	for(int32 FunctionReferenceIndex = 0; FunctionReferenceIndex < FunctionReferencePtrs.Num(); FunctionReferenceIndex++)
 	{
 		TSoftObjectPtr<URigVMFunctionReferenceNode> FunctionReferencePtr = FunctionReferencePtrs[FunctionReferenceIndex];
-		const FString AssetPath = FunctionReferencePtr.ToSoftObjectPath().GetAssetPathName().ToString();
+		const FString AssetPath = FunctionReferencePtr.ToSoftObjectPath().GetAssetPath().ToString();
 		if(AssetPath.StartsWith(TEXT("/Engine/Transient")))
 		{
 			continue;
@@ -12119,7 +12119,7 @@ TArray<FAssetData> URigVMController::GetAffectedAssets(ERigVMControllerBulkEditT
 			continue;
 		}
 					
-		const FAssetData AssetData = AssetRegistryModule.Get().GetAssetByObjectPath(*AssetPath);
+		const FAssetData AssetData = AssetRegistryModule.Get().GetAssetByObjectPath(FSoftObjectPath(AssetPath));
 		if(AssetData.IsValid())
 		{
 			VisitedAssets.Add(AssetPath, Assets.Add(AssetData));
diff --git a/Engine/Source/Developer/SlateReflector/Private/Models/WidgetReflectorNode.cpp b/Engine/Source/Developer/SlateReflector/Private/Models/WidgetReflectorNode.cpp
--- a/Engine/Source/Developer/SlateReflector/Private/Models/WidgetReflectorNode.cpp
+++ b/Engine/Source/Developer/SlateReflector/Private/Models/WidgetReflectorNode.cpp
@@ -537,7 +537,7 @@ TSharedRef<FJsonValue> FSnapshotWidgetReflectorNode::ToJson(const TSharedRef<FSn
 	RootJsonObject->SetField(TEXT("WidgetDesiredSize"), Internal::CreateVector2DJsonValue(RootSnapshotNode->CachedWidgetDesiredSize));
 	RootJsonObject->SetField(TEXT("WidgetForegroundColor"), Internal::CreateSlateColorJsonValue(RootSnapshotNode->CachedWidgetForegroundColor));
 	RootJsonObject->SetStringField(TEXT("WidgetAddress"), Internal::ConvertPtrIntToString(RootSnapshotNode->CachedWidgetAddress));
-	RootJsonObject->SetStringField(TEXT("WidgetAssetPath"), RootSnapshotNode->CachedWidgetAssetData.ObjectPath.ToString());
+	RootJsonObject->SetStringField(TEXT("WidgetAssetPath"), RootSnapshotNode->CachedWidgetAssetData.GetObjectPathString());
 
 	TArray<TSharedPtr<FJsonValue>> ChildNodesJsonArray;
 	for (const auto& ChildReflectorNode : RootSnapshotNode->ChildNodes)
@@ -712,7 +712,7 @@ TSharedRef<FSnapshotWidgetReflectorNode> FSnapshotWidgetReflectorNode::FromJson(
 	RootSnapshotNode->CachedWidgetForegroundColor = Internal::ParseSlateColorJsonValue(RootJsonObject->GetField<EJson::None>(TEXT("WidgetForegroundColor")));
 	RootSnapshotNode->CachedWidgetAddress = Internal::ParsePtrIntFromString(RootJsonObject->GetStringField(TEXT("WidgetAddress")));
 
-	FName AssetPath(*RootJsonObject->GetStringField(TEXT("WidgetAssetPath")));
+	FSoftObjectPath AssetPath(RootJsonObject->GetStringField(TEXT("WidgetAssetPath")));
 	IAssetRegistry& AssetRegistry = FModuleManager::LoadModuleChecked<FAssetRegistryModule>(TEXT("AssetRegistry")).Get();
 	RootSnapshotNode->CachedWidgetAssetData = AssetRegistry.GetAssetByObjectPath(AssetPath);
 
