Add SlateRHIRenderer debuggin option to WidgetReflector

diff --git a/Engine/Source/Developer/SlateReflector/Private/Widgets/SSlateOptions.cpp b/Engine/Source/Developer/SlateReflector/Private/Widgets/SSlateOptions.cpp
--- a/Engine/Source/Developer/SlateReflector/Private/Widgets/SSlateOptions.cpp
+++ b/Engine/Source/Developer/SlateReflector/Private/Widgets/SSlateOptions.cpp
@@ -84,6 +84,9 @@ void SSlateOptions::Construct( const FArguments& InArgs )
 			AddMenuEntry(MenuBuilder, Icon, LOCTEXT("DebugCulling", "Debug Culling"), TEXT("Slate.DebugCulling"));
 			AddMenuEntry(MenuBuilder, Icon, LOCTEXT("ShowHitTestGrid", "Show HitTestGrid"), TEXT("Slate.HitTestGridDebugging"));
 #endif // WITH_SLATE_DEBUGGING
+			AddMenuEntry(MenuBuilder, Icon, LOCTEXT("ShowBatching", "Show Batching"), TEXT("Slate.ShowBatching"));
+			AddMenuEntry(MenuBuilder, Icon, LOCTEXT("ShowOverdraw", "Show Overdraw"), TEXT("Slate.ShowOverdraw"));
+			
 			return MenuBuilder.MakeWidget();
 		}
 	};
