Slate: Add the application scale to the input event visualizer

diff --git a/Engine/Source/Developer/SlateReflector/Private/InputEventVisualizer.cpp b/Engine/Source/Developer/SlateReflector/Private/InputEventVisualizer.cpp
--- a/Engine/Source/Developer/SlateReflector/Private/InputEventVisualizer.cpp
+++ b/Engine/Source/Developer/SlateReflector/Private/InputEventVisualizer.cpp
@@ -232,21 +232,21 @@ void FInputEventVisualizer::HandlePaint(const FPaintArgs& InArgs, const FGeometr
 			}
 
 			int32 NewLayerId = OutLayerId++;
-			const FVector2D CursorSize = FSlateApplication::Get().GetCursorSize();
+			const FVector2f CursorSize = FSlateApplication::Get().GetCursorSize();
 			const FSlateBrush* Brush = FCoreStyle::Get().GetBrush(InputEventVisualizer::CursorPingBrush);
 			for (const FMouseEventInfo& Event : MouseEvents)
 			{
 				// Normalized animation value for the cursor ping between 0 and 1.
 				const float AnimAmount = (float)(CurrentTime - Event.EventTime) / ClickFadeTime;
 
-				const FVector2D CursorPosDesktopSpace = Event.CursorPingPosition;
-				const FVector2D PingSize = CursorSize * InputEventVisualizer::PingScaleAmount * FCurveHandle::ApplyEasing(AnimAmount, ECurveEaseFunction::QuadOut);
+				const FVector2f CursorPosDesktopSpace = Event.CursorPingPosition;
+				const FVector2f PingSize = CursorSize * InputEventVisualizer::PingScaleAmount * FCurveHandle::ApplyEasing(AnimAmount, ECurveEaseFunction::QuadOut);
 				FLinearColor PingColor = Event.PingColor;
 				PingColor.A = 1.0f - FCurveHandle::ApplyEasing(AnimAmount, ECurveEaseFunction::QuadIn);
 
 				FGeometry CursorHighlightGeometry = FGeometry::MakeRoot(PingSize, FSlateLayoutTransform(CursorPosDesktopSpace - PingSize / 2));
 				CursorHighlightGeometry.AppendTransform(Inverse(WindowBeingDrawn->GetLocalToScreenTransform()));
-				CursorHighlightGeometry.AppendTransform(FSlateLayoutTransform(WindowBeingDrawn->GetDPIScaleFactor(), FVector2D::ZeroVector));
+				CursorHighlightGeometry.AppendTransform(FSlateLayoutTransform(WindowBeingDrawn->GetDPIScaleFactor() * FSlateApplicationBase::Get().GetApplicationScale(), FVector2f::ZeroVector));
 
 				FSlateDrawElement::MakeBox(
 					OutDrawElements,
diff --git a/Engine/Source/Developer/SlateReflector/Private/InputEventVisualizer.h b/Engine/Source/Developer/SlateReflector/Private/InputEventVisualizer.h
--- a/Engine/Source/Developer/SlateReflector/Private/InputEventVisualizer.h
+++ b/Engine/Source/Developer/SlateReflector/Private/InputEventVisualizer.h
@@ -54,7 +54,7 @@ private:
 	{
 		FKey Key;
 		double EventTime;
-		FVector2D CursorPingPosition;
+		FVector2f CursorPingPosition;
 		FLinearColor PingColor;
 	};
 	
