Update the widget event log to automatically scroll to the bottom, since scrolling it manually will often result in more events being created

diff --git a/Engine/Source/Developer/SlateReflector/Private/Widgets/SWidgetEventLog.cpp b/Engine/Source/Developer/SlateReflector/Private/Widgets/SWidgetEventLog.cpp
--- a/Engine/Source/Developer/SlateReflector/Private/Widgets/SWidgetEventLog.cpp
+++ b/Engine/Source/Developer/SlateReflector/Private/Widgets/SWidgetEventLog.cpp
@@ -37,7 +37,9 @@ void SWidgetEventLog::Construct(const FArguments& InArgs, TSharedPtr<const SWidg
 	FMessageLogModule& MessageLogModule = FModuleManager::LoadModuleChecked<FMessageLogModule>(TEXT("MessageLog"));
 	if (!MessageLogModule.IsRegisteredLogListing(NAME_WidgetEvents))
 	{
-		MessageLogModule.RegisterLogListing(NAME_WidgetEvents, LOCTEXT("WidgetEventLog", "Widget Events"), FMessageLogInitializationOptions{});
+		FMessageLogInitializationOptions InitOptions;
+		InitOptions.bScrollToBottom = true;
+		MessageLogModule.RegisterLogListing(NAME_WidgetEvents, LOCTEXT("WidgetEventLog", "Widget Events"), InitOptions);
 	}
 	TSharedRef<IMessageLogListing> MessageLogListing = MessageLogModule.GetLogListing(NAME_WidgetEvents);
 
