PR #10997: Optimize EnumHasAnyFlags
Rather than checking if a flag is set to manually negate it, we can xor it every time the button is toggled.

diff --git a/Engine/Source/Developer/SlateReflector/Private/Widgets/SWidgetHittestGrid.cpp b/Engine/Source/Developer/SlateReflector/Private/Widgets/SWidgetHittestGrid.cpp
--- a/Engine/Source/Developer/SlateReflector/Private/Widgets/SWidgetHittestGrid.cpp
+++ b/Engine/Source/Developer/SlateReflector/Private/Widgets/SWidgetHittestGrid.cpp
@@ -648,14 +648,7 @@ TSharedRef<SWidget> SWidgetHittestGrid::GetFlagsMenuContent()
 
 void SWidgetHittestGrid::HandleDisplayFlagsButtonClicked(FHittestGrid::EDisplayGridFlags Flag)
 {
-	if (EnumHasAnyFlags(DisplayGridFlags, Flag))
-	{
-		DisplayGridFlags &= ~Flag;
-	}
-	else
-	{
-		DisplayGridFlags |= Flag;
-	}
+	DisplayGridFlags ^= Flag;
 
 	SaveSettings();
 }
