SlateReflector Truncation warning fixes.

diff --git a/Engine/Source/Developer/SlateReflector/Private/VisualTreeCapture.cpp b/Engine/Source/Developer/SlateReflector/Private/VisualTreeCapture.cpp
--- a/Engine/Source/Developer/SlateReflector/Private/VisualTreeCapture.cpp
+++ b/Engine/Source/Developer/SlateReflector/Private/VisualTreeCapture.cpp
@@ -9,14 +9,14 @@
 #include "Framework/Application/SlateApplication.h"
 #include "Types/InvisibleToWidgetReflectorMetaData.h"
 
-static float VectorSign(const FVector2D& Vec, const FVector2D& A, const FVector2D& B)
+static float VectorSign(const FVector2f& Vec, const FVector2f& A, const FVector2f& B)
 {
 	return FMath::Sign((B.X - A.X) * (Vec.Y - A.Y) - (B.Y - A.Y) * (Vec.X - A.X));
 }
 
 // Returns true when the point is inside the triangle
 // Should not return true when the point is on one of the edges
-static bool IsPointInTriangle(const FVector2D& TestPoint, const FVector2D& A, const FVector2D& B, const FVector2D& C)
+static bool IsPointInTriangle(const FVector2f& TestPoint, const FVector2f& A, const FVector2f& B, const FVector2f& C)
 {
 	float BA = VectorSign(B, A, TestPoint);
 	float CB = VectorSign(C, B, TestPoint);
@@ -38,11 +38,11 @@ FVisualEntry::FVisualEntry(const TWeakPtr<const SWidget>& InWidget, int32 InElem
 FVisualEntry::FVisualEntry(const TSharedRef<const SWidget>& InWidget, const FSlateDrawElement& InElement)
 {
 	const FSlateRenderTransform& Transform = InElement.GetRenderTransform();
-	const FVector2D LocalSize = InElement.GetLocalSize();
+	const FVector2f LocalSize = InElement.GetLocalSize();
 
-	TopLeft = Transform.TransformPoint(FVector2D(0, 0));
-	TopRight = Transform.TransformPoint(FVector2D(LocalSize.X, 0));
-	BottomLeft = Transform.TransformPoint(FVector2D(0, LocalSize.Y));
+	TopLeft = Transform.TransformPoint(FVector2f(0.0f, 0.0f));
+	TopRight = Transform.TransformPoint(FVector2f(LocalSize.X, 0.0f));
+	BottomLeft = Transform.TransformPoint(FVector2f(0.0f, LocalSize.Y));
 	BottomRight = Transform.TransformPoint(LocalSize);
 
 	LayerId = InElement.GetLayer();
@@ -65,11 +65,11 @@ void FVisualEntry::Resolve(const FSlateWindowElementList& ElementList)
 		{
 			const FSlateDrawElement& Element = Container[ElementIndex];
 			const FSlateRenderTransform& Transform = Element.GetRenderTransform();
-			const FVector2D LocalSize = Element.GetLocalSize();
+			const FVector2f LocalSize = Element.GetLocalSize();
 
-			TopLeft = Transform.TransformPoint(FVector2D(0, 0));
-			TopRight = Transform.TransformPoint(FVector2D(LocalSize.X, 0));
-			BottomLeft = Transform.TransformPoint(FVector2D(0, LocalSize.Y));
+			TopLeft = Transform.TransformPoint(FVector2f(0.0f, 0.0f));
+			TopRight = Transform.TransformPoint(FVector2f(LocalSize.X, 0.0f));
+			BottomLeft = Transform.TransformPoint(FVector2f(0.0f, LocalSize.Y));
 			BottomRight = Transform.TransformPoint(LocalSize);
 
 			LayerId = Element.GetLayer();
@@ -80,7 +80,7 @@ void FVisualEntry::Resolve(const FSlateWindowElementList& ElementList)
 	VisitTupleElements(ResolveBounds, ElementList.GetUncachedDrawElements(), UE::Slate::MakeTupleIndicies<uint8, (uint8)EElementType::ET_Count>());
 }
 
-bool FVisualEntry::IsPointInside(const FVector2D& Point) const
+bool FVisualEntry::IsPointInside(const FVector2f& Point) const
 {
 	if (IsPointInTriangle(Point, TopLeft, TopRight, BottomLeft) || IsPointInTriangle(Point, BottomLeft, TopRight, BottomRight))
 	{
@@ -90,7 +90,7 @@ bool FVisualEntry::IsPointInside(const FVector2D& Point) const
 	return false;
 }
 
-TSharedPtr<const SWidget> FVisualTreeSnapshot::Pick(FVector2D Point)
+TSharedPtr<const SWidget> FVisualTreeSnapshot::Pick(FVector2f Point)
 {
 	for (int Index = Entries.Num() - 1; Index >= 0; Index--)
 	{
diff --git a/Engine/Source/Developer/SlateReflector/Private/VisualTreeCapture.h b/Engine/Source/Developer/SlateReflector/Private/VisualTreeCapture.h
--- a/Engine/Source/Developer/SlateReflector/Private/VisualTreeCapture.h
+++ b/Engine/Source/Developer/SlateReflector/Private/VisualTreeCapture.h
@@ -19,10 +19,10 @@ enum class EElementType : uint8;
 class FVisualEntry
 {
 public:
-	FVector2D TopLeft;
-	FVector2D TopRight;
-	FVector2D BottomLeft;
-	FVector2D BottomRight;
+	FVector2f TopLeft;
+	FVector2f TopRight;
+	FVector2f BottomLeft;
+	FVector2f BottomRight;
 
 	int32 LayerId;
 	int32 ClippingIndex;
@@ -36,13 +36,13 @@ public:
 
 	void Resolve(const FSlateWindowElementList& ElementList);
 
-	bool IsPointInside(const FVector2D& Point) const;
+	bool IsPointInside(const FVector2f& Point) const;
 };
 
 class FVisualTreeSnapshot : public TSharedFromThis<FVisualTreeSnapshot>
 {
 public:
-	TSharedPtr<const SWidget> Pick(FVector2D Point);
+	TSharedPtr<const SWidget> Pick(FVector2f Point);
 	
 public:
 	TArray<FVisualEntry> Entries;
diff --git a/Engine/Source/Developer/SlateReflector/Private/Widgets/SAtlasVisualizer.cpp b/Engine/Source/Developer/SlateReflector/Private/Widgets/SAtlasVisualizer.cpp
--- a/Engine/Source/Developer/SlateReflector/Private/Widgets/SAtlasVisualizer.cpp
+++ b/Engine/Source/Developer/SlateReflector/Private/Widgets/SAtlasVisualizer.cpp
@@ -52,7 +52,7 @@ public:
 		const TSharedRef<SWidget>& ChildWidget = ChildSlot.GetWidget();
 		if( ChildWidget->GetVisibility() != EVisibility::Collapsed )
 		{
-			const FVector2D& WidgetDesiredSize = ChildWidget->GetDesiredSize();
+			const FVector2f& WidgetDesiredSize = ChildWidget->GetDesiredSize();
 
 			// Update the zoom level, and clamp the pan offset based on our current geometry
 			SAtlasVisualizerPanel* const NonConstThis = const_cast<SAtlasVisualizerPanel*>(this);
@@ -130,12 +130,12 @@ public:
 			return false;
 		}
 
-		const FVector2D PrevPhysicalOffset = PhysicalOffset;
+		const FVector2f PrevPhysicalOffset = PhysicalOffset;
 		const float InverseZoom = 1.0f / ZoomLevel;
-		PhysicalOffset += (Offset * InverseZoom);
+		PhysicalOffset += (UE::Slate::CastToVector2f(Offset) * InverseZoom);
 
 		const TSharedRef<SWidget>& ChildWidget = ChildSlot.GetWidget();
-		const FVector2D& WidgetDesiredSize = ChildWidget->GetDesiredSize();
+		const FVector2f& WidgetDesiredSize = ChildWidget->GetDesiredSize();
 		ClampViewOffset(WidgetDesiredSize, CachedSize);
 
 		return PhysicalOffset != PrevPhysicalOffset;
@@ -161,7 +161,7 @@ public:
 	void FitToWindow()
 	{
 		bFitToWindow = true;
-		PhysicalOffset = FVector2D::ZeroVector;
+		PhysicalOffset = FVector2f::ZeroVector;
 	}
 
 	bool IsFitToWindow() const
@@ -173,11 +173,11 @@ public:
 	{
 		bFitToWindow = false;
 		ZoomLevel = 1.0f;
-		PhysicalOffset = FVector2D::ZeroVector;
+		PhysicalOffset = FVector2f::ZeroVector;
 	}
 
 private:
-	void UpdateFitToWindowZoom( const FVector2D& ViewportSize, const FVector2D& LocalSize )
+	void UpdateFitToWindowZoom( const FVector2f& ViewportSize, const FVector2f& LocalSize )
 	{
 		if( bFitToWindow )
 		{
@@ -187,7 +187,7 @@ private:
 		}
 	}
 
-	void ClampViewOffset( const FVector2D& ViewportSize, const FVector2D& LocalSize )
+	void ClampViewOffset( const FVector2f& ViewportSize, const FVector2f& LocalSize )
 	{
 		PhysicalOffset.X = ClampViewOffsetAxis(ViewportSize.X, LocalSize.X, PhysicalOffset.X);
 		PhysicalOffset.Y = ClampViewOffsetAxis(ViewportSize.Y, LocalSize.Y, PhysicalOffset.Y);
@@ -222,8 +222,8 @@ private:
 		return CurrentOffset;
 	}
 
-	FVector2D PhysicalOffset;
-	mutable FVector2D CachedSize;
+	FVector2f PhysicalOffset;
+	mutable FVector2f CachedSize;
 	float ZoomLevel;
 	bool bFitToWindow;
 
diff --git a/Engine/Source/Developer/SlateReflector/Private/Widgets/SWidgetReflector.cpp b/Engine/Source/Developer/SlateReflector/Private/Widgets/SWidgetReflector.cpp
--- a/Engine/Source/Developer/SlateReflector/Private/Widgets/SWidgetReflector.cpp
+++ b/Engine/Source/Developer/SlateReflector/Private/Widgets/SWidgetReflector.cpp
@@ -1210,8 +1210,8 @@ int32 SWidgetReflector::Visualize( const FWidgetPath& InWidgetsToVisualize, FSla
 		TSharedPtr<FVisualTreeSnapshot> Tree = VisualCapture.GetVisualTreeForWindow(OutDrawElements.GetPaintWindow());
 		if (Tree.IsValid())
 		{
-			const FVector2D AbsPoint = FSlateApplication::Get().GetCursorPos();
-			const FVector2D WindowPoint = AbsPoint - OutDrawElements.GetPaintWindow()->GetPositionInScreen();
+			const FVector2f AbsPoint = FSlateApplication::Get().GetCursorPos();
+			const FVector2f WindowPoint = AbsPoint - OutDrawElements.GetPaintWindow()->GetPositionInScreen();
 			if (TSharedPtr<const SWidget> PickedWidget = Tree->Pick(WindowPoint))
 			{
 				FWidgetPath WidgetsToVisualize = InWidgetsToVisualize;
@@ -1373,7 +1373,7 @@ void SWidgetReflector::VisualizeAsTree( const TArray<TSharedRef<FWidgetReflector
 			const auto& CurWidget = WidgetPathToVisualize[WidgetIndex];
 
 			// Tint the item based on depth in picked path
-			const float ColorFactor = static_cast<float>(WidgetIndex) / WidgetPathToVisualize.Num();
+			const float ColorFactor = static_cast<float>(WidgetIndex) / static_cast<float>(WidgetPathToVisualize.Num());
 			CurWidget->SetTint(FMath::Lerp(TopmostWidgetColor, LeafmostWidgetColor, ColorFactor));
 
 			// Make sure the user can see the picked path in the tree.
@@ -1398,7 +1398,7 @@ int32 SWidgetReflector::VisualizePickAsRectangles( const FWidgetPath& InWidgetsT
 	for (int32 WidgetIndex = 0; WidgetIndex < InWidgetsToVisualize.Widgets.Num(); ++WidgetIndex)
 	{
 		const FArrangedWidget& WidgetGeometry = InWidgetsToVisualize.Widgets[WidgetIndex];
-		const float ColorFactor = static_cast<float>(WidgetIndex)/InWidgetsToVisualize.Widgets.Num();
+		const float ColorFactor = static_cast<float>(WidgetIndex)/ static_cast<float>(InWidgetsToVisualize.Widgets.Num());
 		const FLinearColor Tint(1.0f - ColorFactor, ColorFactor, 0.0f, 1.0f);
 
 		// The FGeometry we get is from a WidgetPath, so it's rooted in desktop space.
diff --git a/Engine/Source/Developer/SlateReflector/Private/Widgets/SWidgetSnapshotVisualizer.cpp b/Engine/Source/Developer/SlateReflector/Private/Widgets/SWidgetSnapshotVisualizer.cpp
--- a/Engine/Source/Developer/SlateReflector/Private/Widgets/SWidgetSnapshotVisualizer.cpp
+++ b/Engine/Source/Developer/SlateReflector/Private/Widgets/SWidgetSnapshotVisualizer.cpp
@@ -73,7 +73,7 @@ public:
 	{
 		SelectedWindowIndex = InIndex;
 		PickedWidgets.Reset();
-		PhysicalOffset = FVector2D::ZeroVector;
+		PhysicalOffset = FVector2f::ZeroVector;
 	}
 
 	int32 GetSelectedWindowIndex() const
@@ -113,7 +113,7 @@ public:
 		const TSharedRef<SWidget>& ChildWidget = ChildSlot.GetWidget();
 		if (ChildWidget->GetVisibility() != EVisibility::Collapsed)
 		{
-			const FVector2D& WidgetDesiredSize = ChildWidget->GetDesiredSize();
+			const FVector2f& WidgetDesiredSize = ChildWidget->GetDesiredSize();
 
 			// Clamp the pan offset based on our current geometry
 			SScrollableSnapshotImage* const NonConstThis = const_cast<SScrollableSnapshotImage*>(this);
@@ -142,7 +142,7 @@ public:
 	{
 		struct FWidgetPicker
 		{
-			static bool FindWidgetsUnderPoint(const FVector2D& InHitTestPoint, const FVector2D& InWindowPosition, const TSharedRef<FWidgetReflectorNodeBase>& InWidget, TArray<TSharedRef<FWidgetReflectorNodeBase>>& OutWidgets)
+			static bool FindWidgetsUnderPoint(const FVector2f& InHitTestPoint, const FVector2f& InWindowPosition, const TSharedRef<FWidgetReflectorNodeBase>& InWidget, TArray<TSharedRef<FWidgetReflectorNodeBase>>& OutWidgets)
 			{
 				const bool bNeedsHitTesting = InWidget->GetHitTestInfo().IsHitTestVisible || InWidget->GetHitTestInfo().AreChildrenHitTestVisible;
 				if (bNeedsHitTesting)
@@ -178,9 +178,9 @@ public:
 		if (bIsPicking)
 		{
 			// We need to pick in the snapshot window space, so convert the mouse co-ordinates to be relative to our top-left position
-			const FVector2D& ScreenMousePos = MouseEvent.GetScreenSpacePosition();
-			const FVector2D LocalMousePos = MyGeometry.AbsoluteToLocal(ScreenMousePos);
-			const FVector2D ScrolledPos = LocalMousePos - PhysicalOffset;
+			const FVector2f& ScreenMousePos = MouseEvent.GetScreenSpacePosition();
+			const FVector2f LocalMousePos = MyGeometry.AbsoluteToLocal(ScreenMousePos);
+			const FVector2f ScrolledPos = LocalMousePos - PhysicalOffset;
 
 			PickedWidgets.Reset();
 
@@ -248,7 +248,7 @@ public:
 				for (int32 WidgetIndex = 0; WidgetIndex < PickedWidgets.Num(); ++WidgetIndex)
 				{
 					const TSharedRef<FWidgetReflectorNodeBase>& PickedWidget = PickedWidgets[WidgetIndex];
-					const float ColorFactor = static_cast<float>(WidgetIndex)/PickedWidgets.Num();
+					const float ColorFactor = static_cast<float>(WidgetIndex)/ static_cast<float>(PickedWidgets.Num());
 					const FLinearColor Tint(1.0f - ColorFactor, ColorFactor, 0.0f, 1.0f);
 
 					FSlateDrawElement::MakeBox(
@@ -288,11 +288,11 @@ public:
 
 	virtual bool ScrollBy(const FVector2D& Offset) override
 	{
-		const FVector2D PrevPhysicalOffset = PhysicalOffset;
-		PhysicalOffset += Offset;
+		const FVector2f PrevPhysicalOffset = PhysicalOffset;
+		PhysicalOffset += UE::Slate::CastToVector2f(Offset);
 
 		const TSharedRef<SWidget>& ChildWidget = ChildSlot.GetWidget();
-		const FVector2D& WidgetDesiredSize = ChildWidget->GetDesiredSize();
+		const FVector2f& WidgetDesiredSize = ChildWidget->GetDesiredSize();
 		ClampViewOffset(WidgetDesiredSize, CachedSize);
 
 		return PhysicalOffset != PrevPhysicalOffset;
@@ -309,7 +309,7 @@ public:
 	}
 
 private:
-	void ClampViewOffset(const FVector2D& ViewportSize, const FVector2D& LocalSize)
+	void ClampViewOffset(const FVector2f& ViewportSize, const FVector2f& LocalSize)
 	{
 		PhysicalOffset.X = ClampViewOffsetAxis(ViewportSize.X, LocalSize.X, PhysicalOffset.X);
 		PhysicalOffset.Y = ClampViewOffsetAxis(ViewportSize.Y, LocalSize.Y, PhysicalOffset.Y);
@@ -342,8 +342,8 @@ private:
 		return CurrentOffset;
 	}
 
-	FVector2D PhysicalOffset;
-	mutable FVector2D CachedSize;
+	FVector2f PhysicalOffset;
+	mutable FVector2f CachedSize;
 
 	FScrollyZoomy ScrollyZoomy;
 
diff --git a/Engine/Source/Developer/SlateReflector/SlateReflector.Build.cs b/Engine/Source/Developer/SlateReflector/SlateReflector.Build.cs
--- a/Engine/Source/Developer/SlateReflector/SlateReflector.Build.cs
+++ b/Engine/Source/Developer/SlateReflector/SlateReflector.Build.cs
@@ -93,5 +93,7 @@ public class SlateReflector : ModuleRules
 		{
 			PrecompileForTargets = PrecompileTargetsType.Any;
 		}
+
+		UnsafeTypeCastWarningLevel = WarningLevel.Error;
 	}
 }
