Widget Reflector:
- Add flag to disable the default behavior of "focus on pick"

diff --git a/Engine/Source/Developer/SlateReflector/Private/Widgets/SSlateOptions.cpp b/Engine/Source/Developer/SlateReflector/Private/Widgets/SSlateOptions.cpp
--- a/Engine/Source/Developer/SlateReflector/Private/Widgets/SSlateOptions.cpp
+++ b/Engine/Source/Developer/SlateReflector/Private/Widgets/SSlateOptions.cpp
@@ -62,6 +62,7 @@ void SSlateOptions::Construct( const FArguments& InArgs )
 			AddMenuEntry(MenuBuilder, Icon, LOCTEXT("EnableToolTips", "Enable Tooltips"), TEXT("Slate.EnableTooltips"));
 			AddMenuEntry(MenuBuilder, Icon, LOCTEXT("GlobalInvalidation", "Global Invalidation"), TEXT("Slate.EnableGlobalInvalidation"));
 			AddMenuEntry(MenuBuilder, Icon, LOCTEXT("DisabledEffect", "Transparent Disabled Effect"), TEXT("Slate.ApplyDisabledEffectOnWidgets"));
+			AddMenuEntry(MenuBuilder, Icon, LOCTEXT("FocusOnPick", "Focus On Pick"), TEXT("Slate.EnableFocusOnPick"));
 
 			return MenuBuilder.MakeWidget();
 		}
diff --git a/Engine/Source/Developer/SlateReflector/Private/Widgets/SWidgetReflector.cpp b/Engine/Source/Developer/SlateReflector/Private/Widgets/SWidgetReflector.cpp
--- a/Engine/Source/Developer/SlateReflector/Private/Widgets/SWidgetReflector.cpp
+++ b/Engine/Source/Developer/SlateReflector/Private/Widgets/SWidgetReflector.cpp
@@ -71,6 +71,16 @@
  * Widget reflector user widget.
  */
 
+namespace WidgetReflectorCVars
+{
+	static bool bEnableFocusOnPick = true;
+	static FAutoConsoleVariableRef EnableFocusOnPick(
+		TEXT("Slate.EnableFocusOnPick"),
+		bEnableFocusOnPick,
+		TEXT("If true, Widget Reflector window will be automatically focused when a pick is made.")
+	);
+}
+
 /* Local helpers
  *****************************************************************************/
  
@@ -2020,7 +2030,11 @@ void SWidgetReflector::HandleReflectorTreeSelectionChanged( TSharedPtr<FWidgetRe
 
 	if (GIsEditor && SelectedWidgetObjects.Num() > 0)
 	{
-		TabManager->TryInvokeTab(WidgetReflectorTabID::WidgetDetails);
+		if (WidgetReflectorCVars::bEnableFocusOnPick)
+		{
+			TabManager->TryInvokeTab(WidgetReflectorTabID::WidgetDetails);
+		}
+
 		if (PropertyViewPtr.IsValid())
 		{
 			PropertyViewPtr->SetObjects(SelectedWidgetObjects);
