Slate: Add missing validation test to the WidgetReflector.
#rnx

diff --git a/Engine/Source/Developer/SlateReflector/Private/Widgets/SSlateOptions.cpp b/Engine/Source/Developer/SlateReflector/Private/Widgets/SSlateOptions.cpp
--- a/Engine/Source/Developer/SlateReflector/Private/Widgets/SSlateOptions.cpp
+++ b/Engine/Source/Developer/SlateReflector/Private/Widgets/SSlateOptions.cpp
@@ -118,9 +118,11 @@ void SSlateOptions::Construct( const FArguments& InArgs )
 				AddMenuEntry(MenuBuilder, Icon, LOCTEXT("VerifyWidgetIndex", "Verify Widget's index"), TEXT("Slate.InvalidationRoot.VerifyWidgetsIndex"));
 				AddMenuEntry(MenuBuilder, Icon, LOCTEXT("VerifyWidgetPtr", "Verify Widget Pointer"), TEXT("Slate.InvalidationRoot.VerifyValidWidgets"));
 				AddMenuEntry(MenuBuilder, Icon, LOCTEXT("VerifyHittestGrid", "Verify Hittest Grid"), TEXT("Slate.InvalidationRoot.VerifyHittestGrid"));
+				AddMenuEntry(MenuBuilder, Icon, LOCTEXT("VerifyCachedElement", "Verify Cached Elements"), TEXT("Slate.InvalidationRoot.VerifyCachedElementDataList"));
 				AddMenuEntry(MenuBuilder, Icon, LOCTEXT("VefiryVisibility", "Verify Visibility"), TEXT("Slate.InvalidationRoot.VerifyWidgetVisibility"));
 				AddMenuEntry(MenuBuilder, Icon, LOCTEXT("VerifyVolatility", "Verify Volatility"), TEXT("Slate.InvalidationRoot.VerifyWidgetVolatile"));
 				AddMenuEntry(MenuBuilder, Icon, LOCTEXT("VerifyUpdateList", "Verify Update List"), TEXT("Slate.InvalidationRoot.VerifyWidgetUpdateList"));
+				AddMenuEntry(MenuBuilder, Icon, LOCTEXT("VerifyUpdateOnce", "Verify Updated Once"), TEXT("Slate.InvalidationRoot.VerifyWidgetsAreUpdatedOnce"));
 				AddMenuEntry(MenuBuilder, Icon, LOCTEXT("VerifySlateAttributes", "Verify Attributes"), TEXT("Slate.InvalidationRoot.VerifySlateAttribute"));
 	#endif // UE_SLATE_WITH_INVALIDATIONWIDGETLIST_DEBUGGING
 				MenuBuilder.EndSection();
