Fix panning of texture atlas visualizer

diff --git a/Engine/Source/Developer/SlateReflector/Private/Widgets/SAtlasVisualizer.cpp b/Engine/Source/Developer/SlateReflector/Private/Widgets/SAtlasVisualizer.cpp
--- a/Engine/Source/Developer/SlateReflector/Private/Widgets/SAtlasVisualizer.cpp
+++ b/Engine/Source/Developer/SlateReflector/Private/Widgets/SAtlasVisualizer.cpp
@@ -432,7 +432,7 @@ void SAtlasVisualizer::OnMouseLeave(const FPointerEvent& MouseEvent)
 FReply SAtlasVisualizer::OnMouseMove(const FGeometry& MyGeometry, const FPointerEvent& MouseEvent)
 {
 #if WITH_ATLAS_DEBUGGING
-	if (!MouseEvent.GetCursorDelta().IsNearlyZero() && !HasMouseCapture())
+	if (!MouseEvent.GetCursorDelta().IsNearlyZero() && !HasMouseCapture() && !MouseEvent.IsMouseButtonDown(EKeys::RightMouseButton))
 	{
 		FVector2D LocalPos = MyGeometry.AbsoluteToLocal(MouseEvent.GetScreenSpacePosition()) / ScrollPanel->GetZoomLevel();
 
