Fix deprecation code in WidgetList introduced with CL 33424099

diff --git a/Engine/Source/Developer/SlateReflector/Private/Widgets/SWidgetList.cpp b/Engine/Source/Developer/SlateReflector/Private/Widgets/SWidgetList.cpp
--- a/Engine/Source/Developer/SlateReflector/Private/Widgets/SWidgetList.cpp
+++ b/Engine/Source/Developer/SlateReflector/Private/Widgets/SWidgetList.cpp
@@ -174,7 +174,6 @@ void SWidgetList::Construct(const FArguments& Args)
 		+ SVerticalBox::Slot()
 		[
 			SAssignNew(ListView, SListView<TSharedPtr<FWidgetListEntry>>)
-			.ItemHeight(24)
 			.ListItemsSource(&Entries)
 			.OnGenerateRow(this, &SWidgetList::OnGenerateRow)
 			.HeaderRow
