Localization:
Introduce the WidgetReflector.DumpFontAtlases console command.
- This console command will save all of the atlased font textures in the font cache to disk.
- By default, the commadn will save the textures to Project/Saved/WidgetReflector/FontAtlases_<TAMESTAMP> directory.
- It is possible to provide an absolute save directory for the files to be saved instead using the -SaveDirectory parameter. Note the timestamp will still be appended to this provided file path.
- It is possible to provide a prefix to the save files using the -SaveFilePrefix parameter. Note the prefix does not automatically append any hyphens or underscores to the file name, the user must provide it themselves.
- All dumped textures should be correctly gamma corrected (where necessary) and using the correct color channels and bit depth.

diff --git a/Engine/Source/Developer/SlateReflector/Private/Widgets/SWidgetReflector.cpp b/Engine/Source/Developer/SlateReflector/Private/Widgets/SWidgetReflector.cpp
--- a/Engine/Source/Developer/SlateReflector/Private/Widgets/SWidgetReflector.cpp
+++ b/Engine/Source/Developer/SlateReflector/Private/Widgets/SWidgetReflector.cpp
@@ -49,6 +49,12 @@
 #include "Debugging/SlateDebugging.h"
 #include "VisualTreeCapture.h"
 #include "Models/WidgetReflectorNode.h"
+#include "Fonts/FontCache.h"
+#include "Fonts/FontTypes.h"
+#include "IImageWrapperModule.h"
+#include "Misc/FileHelper.h"
+#include "Misc/Paths.h"
+#include "Textures/SlateShaderResource.h"
 
 #if SLATE_REFLECTOR_HAS_DESKTOP_PLATFORM
 #include "DesktopPlatformModule.h"
@@ -89,11 +95,18 @@ namespace WidgetReflectorImpl
 /** Command to take a snapshot. */
 void TakeSnapshotCommand(const TArray<FString>&);
 
+void DumpFontAtlasesCommand(const TArray<FString>&);
+
 FAutoConsoleCommand CCommandWidgetReflectorTakeSnapshot(
 	TEXT("WidgetReflector.TakeSnapshot")
 	, TEXT("Take a snapshot and save the result on the local drive. ie. WidgetReflector.TakeSnapshot [Delay=1.0] [Navigation=false]")
 	, FConsoleCommandWithArgsDelegate::CreateStatic(&TakeSnapshotCommand));
 
+FAutoConsoleCommand CCommandDumpFontAtlases(
+	TEXT("WidgetReflector.DumpFontAtlases")
+	, TEXT("Save all of the font atlases currently in the font cache on the local drive. A timestamp will be automatically appended to the directory where the textures will be saved. The default location is YourProject/Saved/WidgetReflector/FontAtlases_Timestamp ie. WidgetReflector.DumpFontAtlases[SaveDirectory=D:\\My\\Path] [SaveFilePrefix=ExamplePrefix]")
+	, FConsoleCommandWithArgsDelegate::CreateStatic(&DumpFontAtlasesCommand));
+
 /** Information about a potential widget snapshot target */
 struct FWidgetSnapshotTarget
 {
@@ -2234,6 +2247,92 @@ void TakeSnapshotCommand(const TArray<FString>& Args)
 	TakeSnapshotEndFrameHandle = FCoreDelegates::OnEndFrame.AddStatic(&TakeSnapshotCommand_EndFrame, RequestedTimeDelay, bRequestNavigation);
 }
 
+void DumpFontAtlasesCommand(const TArray<FString>& Args)
+{
+	FString SaveDirectoryName = FPaths::ProjectSavedDir() / TEXT("WidgetReflector") / TEXT("FontAtlases");
+	FString SaveFilePrefix;
+	for (const FString& Arg : Args)
+	{
+		if (FParse::Value(*Arg, TEXT("SaveDirectory"), SaveDirectoryName))
+		{
+			continue;
+		}
+		if (FParse::Value(*Arg, TEXT("SaveFilePrefix"), SaveFilePrefix))
+		{
+			continue;
+		}
+	}
+	
+	// We append the timestamp on to the directory names to differentiate between various runs of the command.
+	SaveDirectoryName = FString::Format(TEXT("{0}_{1}"), {SaveDirectoryName, FDateTime::Now().ToString()});
+	SaveDirectoryName = FPaths::ConvertRelativePathToFull(SaveDirectoryName);
+	UE_LOG(LogSlate, Display, TEXT("Dumping all current font atlases to directory '%s'. All textures will be saved with a user provided file prefix of '%s'"), *SaveDirectoryName, *SaveFilePrefix);
+	IImageWrapperModule& ImageWrapperModule = FModuleManager::LoadModuleChecked<IImageWrapperModule>(FName("ImageWrapper"));
+	TSharedRef<FSlateFontCache>  FontCache = FSlateApplication::Get().GetRenderer()->GetFontCache();
+	int32 NumFontAtlases = FontCache->GetNumAtlasPages();
+	for (int32 Index = 0; Index < NumFontAtlases; ++Index)
+	{
+		// @TODOFonts: Refactor this out into its own function so the widget reflector could have a button to save a single texture page or save all texture pages 
+		ISlateFontTexture* FontTexture = FontCache->GetFontTexture(Index);
+		if (!FontTexture)
+		{
+			UE_LOG(LogSlate, Warning, TEXT("Font texture page %d is not a valid font texture. This page will not be saved."), Index);
+			continue;
+		}
+
+		TArray<uint8> FontTextureData;
+		FontTexture->GetAtlasDataCopy(FontTextureData);
+		if (FontTextureData.IsEmpty())
+		{
+			UE_LOG(LogSlate, Log, TEXT("Texture page %d has no texture data. The texture will not be saved."), Index);
+			continue;
+		}
+		// We need this to get access to the height and width of the texture 
+		FSlateShaderResource* ShaderResource = FontTexture->GetSlateTexture();
+		if (!ShaderResource)
+		{
+			UE_LOG(LogSlate, Log, TEXT("Unable to retrieve the shader resource for texture page %d. Skipping saving this texture."), Index);
+			continue;
+		}
+		ESlateFontAtlasContentType ContentType = FontTexture->GetContentType();
+		ERawImageFormat::Type ImageFormat = ERawImageFormat::Invalid;
+		EGammaSpace GammaSpace = EGammaSpace::Invalid;
+		switch (ContentType)
+		{
+		case ESlateFontAtlasContentType::Alpha:
+			ImageFormat = ERawImageFormat::G8;
+			GammaSpace = EGammaSpace::Linear;
+			break;
+		case ESlateFontAtlasContentType::Color:
+			ImageFormat = ERawImageFormat::BGRA8;
+			GammaSpace = EGammaSpace::sRGB;
+			break;
+		case ESlateFontAtlasContentType::Msdf:
+			ImageFormat = ERawImageFormat::BGRA8;
+			GammaSpace = EGammaSpace::Linear;
+			break;
+		default:
+			break;
+		}
+		FImageView ImageView((void*) FontTextureData.GetData(), ShaderResource->GetWidth(), ShaderResource->GetHeight(), 1 /* InNumSlices*/, ImageFormat, GammaSpace);
+		TArray64<uint8> SaveImageBuffer;
+		bool bSavedImage = ImageWrapperModule.CompressImage(SaveImageBuffer, EImageFormat::PNG, ImageView);
+		if (bSavedImage)
+		{
+			FString SaveFileName = FString::Format(TEXT("{0}FontAtlasPage_{1}.png"), { SaveFilePrefix, FString::FromInt(Index)});
+			FString AbsoluteSaveFileName = SaveDirectoryName / SaveFileName;
+			UE_LOG(LogSlate, Display, TEXT("Saving font texture '%s'"), *AbsoluteSaveFileName);
+			if (!FFileHelper::SaveArrayToFile(SaveImageBuffer, *AbsoluteSaveFileName))
+			{
+				UE_LOG(LogSlate, Warning, TEXT("Failed to save font texture '%s'"), *AbsoluteSaveFileName);
+			}
+		}
+		else
+		{
+			UE_LOG(LogSlate, Warning, TEXT("Failed to convert texture page %d to a buffer to be saved."), Index);
+		}
+	}
+}
 } // namespace WidgetReflectorImpl
 
 TSharedRef<SWidgetReflector> SWidgetReflector::New()
diff --git a/Engine/Source/Developer/SlateReflector/SlateReflector.Build.cs b/Engine/Source/Developer/SlateReflector/SlateReflector.Build.cs
--- a/Engine/Source/Developer/SlateReflector/SlateReflector.Build.cs
+++ b/Engine/Source/Developer/SlateReflector/SlateReflector.Build.cs
@@ -24,7 +24,9 @@ public class SlateReflector : ModuleRules
 				"AssetRegistry",
                 "MessageLog",
 				"ToolWidgets",
-				"DesktopPlatform"
+				"DesktopPlatform",
+				"ImageCore",
+				"ImageWrapper"
             }
 		);
 
diff --git a/Engine/Source/Runtime/SlateCore/Private/Fonts/FontTypes.cpp b/Engine/Source/Runtime/SlateCore/Private/Fonts/FontTypes.cpp
--- a/Engine/Source/Runtime/SlateCore/Private/Fonts/FontTypes.cpp
+++ b/Engine/Source/Runtime/SlateCore/Private/Fonts/FontTypes.cpp
@@ -17,6 +17,11 @@ ESlateFontAtlasContentType FSlateFontAtlas::GetContentType() const
 	return ContentType;
 }
 
+void FSlateFontAtlas::GetAtlasDataCopy(TArray<uint8>& OutData) const
+{
+	OutData = AtlasData;
+}
+
 const FAtlasedTextureSlot* FSlateFontAtlas::AddCharacter( const FCharacterRenderData& RenderData )
 {
 	check(RenderData.ContentType == GetContentType());
diff --git a/Engine/Source/Runtime/SlateCore/Public/Fonts/FontTypes.h b/Engine/Source/Runtime/SlateCore/Public/Fonts/FontTypes.h
--- a/Engine/Source/Runtime/SlateCore/Public/Fonts/FontTypes.h
+++ b/Engine/Source/Runtime/SlateCore/Public/Fonts/FontTypes.h
@@ -156,6 +156,14 @@ public:
 	 * Releases rendering resources of this texture
 	 */
 	virtual void ReleaseRenderingResources() = 0;
+
+	/**
+ * Retrieves a copy of the texture atlas data in this texture.
+ * This is a slow operation so this is recommended primarily for debugging and visualization tools.
+ *
+ * @param OutData	An array that a caller provides. This array will be populated with a copy of the data in the atlas. It is up to the caller to verify that the array is populated and not empty.
+ */
+	virtual void GetAtlasDataCopy(TArray<uint8>& OutData) const = 0;
 };
 
 /** 
@@ -167,10 +175,12 @@ public:
 	SLATECORE_API FSlateFontAtlas(uint32 InWidth, uint32 InHeight, ESlateFontAtlasContentType InContentType, ESlateTextureAtlasPaddingStyle InPaddingStyle);
 	SLATECORE_API virtual ~FSlateFontAtlas();
 
-	//~ ISlateFontTexture interface
+	//~ Begin ISlateFontTexture interface 
 	SLATECORE_API virtual ESlateFontAtlasContentType GetContentType() const override final;
 	virtual FSlateShaderResource* GetAtlasTexture() const override { return GetSlateTexture(); }
 	virtual void ReleaseRenderingResources() { ReleaseResources(); }
+	SLATECORE_API virtual void GetAtlasDataCopy(TArray<uint8>& OutData) const override;
+	//~ End ISlateFontTexture interface
 
 	/**
 	 * Flushes all cached data.
@@ -206,6 +216,7 @@ protected:
 class ISlateFontAtlasFactory
 {
 public:
+	virtual ~ISlateFontAtlasFactory() {}
 	virtual FIntPoint GetAtlasSize(ESlateFontAtlasContentType InContentType) const = 0;
 	virtual TSharedRef<FSlateFontAtlas> CreateFontAtlas(ESlateFontAtlasContentType InContentType) const = 0;
 	virtual TSharedPtr<ISlateFontTexture> CreateNonAtlasedTexture(const uint32 InWidth, const uint32 InHeight, ESlateFontAtlasContentType InContentType, const TArray<uint8>& InRawData) const = 0;
diff --git a/Engine/Source/Runtime/SlateRHIRenderer/Private/SlateRHIFontTexture.cpp b/Engine/Source/Runtime/SlateRHIRenderer/Private/SlateRHIFontTexture.cpp
--- a/Engine/Source/Runtime/SlateRHIRenderer/Private/SlateRHIFontTexture.cpp
+++ b/Engine/Source/Runtime/SlateRHIRenderer/Private/SlateRHIFontTexture.cpp
@@ -200,6 +200,12 @@ FSlateFontTextureRHI::~FSlateFontTextureRHI()
 {
 }
 
+void FSlateFontTextureRHI::GetAtlasDataCopy(TArray<uint8>& OutData) const
+{
+	OutData.Empty();
+	// @TODOFonts: Implement 
+}
+
 void FSlateFontTextureRHI::ReleaseResources()
 {
 	checkSlow(IsThreadSafeForSlateRendering());
diff --git a/Engine/Source/Runtime/SlateRHIRenderer/Private/SlateRHIFontTexture.h b/Engine/Source/Runtime/SlateRHIRenderer/Private/SlateRHIFontTexture.h
--- a/Engine/Source/Runtime/SlateRHIRenderer/Private/SlateRHIFontTexture.h
+++ b/Engine/Source/Runtime/SlateRHIRenderer/Private/SlateRHIFontTexture.h
@@ -90,6 +90,7 @@ public:
 	{
 		ReleaseResources();
 	}
+	virtual void GetAtlasDataCopy(TArray<uint8>& OutData) const override;
 
 	void ReleaseResources();
 private:
