SWidgetReflector: Prevent issue with PropertyDetails preventing UWidgets from being GC'ed during PIE shutdown.

#rnx

diff --git a/Engine/Source/Developer/SlateReflector/Private/Widgets/SWidgetReflector.cpp b/Engine/Source/Developer/SlateReflector/Private/Widgets/SWidgetReflector.cpp
--- a/Engine/Source/Developer/SlateReflector/Private/Widgets/SWidgetReflector.cpp
+++ b/Engine/Source/Developer/SlateReflector/Private/Widgets/SWidgetReflector.cpp
@@ -66,6 +66,7 @@
 #endif // SLATE_REFLECTOR_HAS_SESSION_SERVICES
 
 #if WITH_EDITOR
+#include "Editor.h"
 #include "Framework/Docking/LayoutService.h"
 #include "PropertyEditorModule.h"
 #endif
@@ -180,12 +181,15 @@ class SWidgetReflector : public ::SWidgetReflector
 	// The reflector uses a tree that observes FWidgetReflectorNodeBase objects.
 	typedef STreeView<TSharedRef<FWidgetReflectorNodeBase>> SReflectorTree;
 
-private:
-
 	//~ Begin ::SWidgetReflector implementation
+public:
+	virtual ~SWidgetReflector();
+
+private:
 	virtual void Construct( const FArguments& InArgs) override;
 	//~ End ::SWidgetReflector implementation
 
+
 	void HandlePullDownAtlasesMenu(FMenuBuilder& MenuBuilder);
 	void HandlePullDownWindowMenu(FMenuBuilder& MenuBuilder);
 
@@ -433,6 +437,10 @@ public:
 	/** Determine the text of Take Snapshot button */
 	FText HandleGetTakeSnapshotText() const { return bIsPendingDelayedSnapshot ? LOCTEXT("CancelSnapshotButtonText", "Cancel Snapshot") : LOCTEXT("TakeSnapshotButtonText", "Take Snapshot"); }
 
+#if WITH_EDITOR
+	void HandlePIEShutdown(const bool bIsSimulating);
+#endif //WITH_EDITOR
+
 private:
 	TSharedPtr<FTabManager> TabManager;
 	TMap<FName, TWeakPtr<SDockTab>> SpawnedTabs;
@@ -671,6 +679,17 @@ void SWidgetReflector::Construct( const FArguments& InArgs )
 			]
 		]
 	];
+
+#if WITH_EDITOR
+	FEditorDelegates::ShutdownPIE.AddSP(this, &SWidgetReflector::HandlePIEShutdown);
+#endif //WITH_EDITOR
+}
+
+SWidgetReflector::~SWidgetReflector()
+{
+#if WITH_EDITOR
+	FEditorDelegates::ShutdownPIE.RemoveAll(this);
+#endif //WITH_EDITOR
 }
 
 void SWidgetReflector::HandlePullDownAtlasesMenu(FMenuBuilder& MenuBuilder)
@@ -2191,6 +2210,16 @@ EVisibility SWidgetReflector::HandleIsBreadcrumbVisible() const
 	return BreadCrumb->HasCrumbs() ? EVisibility::Visible : EVisibility::Collapsed;
 }
 
+#if WITH_EDITOR
+void SWidgetReflector::HandlePIEShutdown(const bool bIsSimulating)
+{
+	if (PropertyViewPtr.IsValid())
+	{
+		PropertyViewPtr->SetObject(nullptr);
+	}
+}
+#endif //WITH_EDITOR
+
 void SWidgetReflector::HandleResetFilteredTreeRoot()
 {
 	BreadCrumb->ClearCrumbs();
diff --git a/Engine/Source/Developer/SlateReflector/Private/Widgets/SWidgetReflector.h b/Engine/Source/Developer/SlateReflector/Private/Widgets/SWidgetReflector.h
--- a/Engine/Source/Developer/SlateReflector/Private/Widgets/SWidgetReflector.h
+++ b/Engine/Source/Developer/SlateReflector/Private/Widgets/SWidgetReflector.h
@@ -29,5 +29,6 @@ public:
 	SLATE_END_ARGS()
 
 	virtual void Construct( const FArguments& InArgs ) = 0;
+	virtual ~SWidgetReflector() = default;
 };
 
