WidgetReflector: When closing PIE, remove any transaction buffer records that contains references to gameplay objects

#rnx

diff --git a/Engine/Source/Developer/SlateReflector/Private/Widgets/SWidgetReflector.cpp b/Engine/Source/Developer/SlateReflector/Private/Widgets/SWidgetReflector.cpp
--- a/Engine/Source/Developer/SlateReflector/Private/Widgets/SWidgetReflector.cpp
+++ b/Engine/Source/Developer/SlateReflector/Private/Widgets/SWidgetReflector.cpp
@@ -438,7 +438,8 @@ public:
 	FText HandleGetTakeSnapshotText() const { return bIsPendingDelayedSnapshot ? LOCTEXT("CancelSnapshotButtonText", "Cancel Snapshot") : LOCTEXT("TakeSnapshotButtonText", "Take Snapshot"); }
 
 #if WITH_EDITOR
-	void HandlePIEShutdown(const bool bIsSimulating);
+	void HandleEndPIE(const bool bIsSimulating);
+	void HandleObjectTransacted(UObject* InObject, const FTransactionObjectEvent& InTransactionEvent);
 #endif //WITH_EDITOR
 
 private:
@@ -491,6 +492,9 @@ private:
 	float SnapshotDelay;
 	bool bIsPendingDelayedSnapshot;
 	bool bRequestNavigationSimulation;
+#if WITH_EDITOR
+	bool bTransactedPIEObject;
+#endif
 	double TimeOfScheduledSnapshot;
 };
 
@@ -513,6 +517,9 @@ void SWidgetReflector::Construct( const FArguments& InArgs )
 	SnapshotDelay = 0.0f;
 	bIsPendingDelayedSnapshot = false;
 	bRequestNavigationSimulation = false;
+#if WITH_EDITOR
+	bTransactedPIEObject = false;
+#endif 
 	TimeOfScheduledSnapshot = -1.0;
 
 	WidgetSnapshotService = InArgs._WidgetSnapshotService;
@@ -681,14 +688,16 @@ void SWidgetReflector::Construct( const FArguments& InArgs )
 	];
 
 #if WITH_EDITOR
-	FEditorDelegates::ShutdownPIE.AddSP(this, &SWidgetReflector::HandlePIEShutdown);
+	FEditorDelegates::EndPIE.AddSP(this, &SWidgetReflector::HandleEndPIE);
+	FCoreUObjectDelegates::OnObjectTransacted.AddRaw(this, &SWidgetReflector::HandleObjectTransacted);
 #endif //WITH_EDITOR
 }
 
 SWidgetReflector::~SWidgetReflector()
 {
 #if WITH_EDITOR
-	FEditorDelegates::ShutdownPIE.RemoveAll(this);
+	FEditorDelegates::EndPIE.RemoveAll(this);
+	FCoreUObjectDelegates::OnObjectTransacted.RemoveAll(this);
 #endif //WITH_EDITOR
 }
 
@@ -2211,13 +2220,32 @@ EVisibility SWidgetReflector::HandleIsBreadcrumbVisible() const
 }
 
 #if WITH_EDITOR
-void SWidgetReflector::HandlePIEShutdown(const bool bIsSimulating)
+void SWidgetReflector::HandleEndPIE(const bool bIsSimulating)
 {
 	if (PropertyViewPtr.IsValid())
 	{
 		PropertyViewPtr->SetObject(nullptr);
 	}
+
+	if (bTransactedPIEObject)
+	{
+		GEditor->ResetTransaction(LOCTEXT("WidgetReflectorDuringPIE", "PIE Objects may be in the transaction buffer that need to be destroyed"));
+		bTransactedPIEObject = false;
+	}
+}
+
+void SWidgetReflector::HandleObjectTransacted(UObject* InObject, const FTransactionObjectEvent& InTransactionEvent)
+{
+	// If the transacted object was recorded by the widget reflector.
+	if (GIsPlayInEditorWorld)
+	{
+		if (PropertyViewPtr && PropertyViewPtr->GetSelectedObjects().Contains(InObject))
+		{
+			bTransactedPIEObject = true;
+		}
+	}
 }
+
 #endif //WITH_EDITOR
 
 void SWidgetReflector::HandleResetFilteredTreeRoot()
