WidgetReflector: Fix issue with nontransacted property changes causing PIE crash.

#rnx

diff --git a/Engine/Source/Developer/SlateReflector/Private/Widgets/SWidgetReflector.cpp b/Engine/Source/Developer/SlateReflector/Private/Widgets/SWidgetReflector.cpp
--- a/Engine/Source/Developer/SlateReflector/Private/Widgets/SWidgetReflector.cpp
+++ b/Engine/Source/Developer/SlateReflector/Private/Widgets/SWidgetReflector.cpp
@@ -2237,7 +2237,7 @@ void SWidgetReflector::HandleEndPIE(const bool bIsSimulating)
 void SWidgetReflector::HandleObjectTransacted(UObject* InObject, const FTransactionObjectEvent& InTransactionEvent)
 {
 	// If the transacted object was recorded by the widget reflector.
-	if (GIsPlayInEditorWorld)
+	if ((GEditor != nullptr && GEditor->PlayWorld != nullptr) || GIsPlayInEditorWorld)
 	{
 		if (PropertyViewPtr && PropertyViewPtr->GetSelectedObjects().Contains(InObject))
 		{
@@ -2245,7 +2245,6 @@ void SWidgetReflector::HandleObjectTransacted(UObject* InObject, const FTransact
 		}
 	}
 }
-
 #endif //WITH_EDITOR
 
 void SWidgetReflector::HandleResetFilteredTreeRoot()
