Minor Dynamic Invalidation Fixes.
We also decided to not expose the `bUseDynamicInvalidation` property to UMG at this point.

diff --git a/Engine/Source/Developer/SlateReflector/Private/Models/WidgetReflectorNode.cpp b/Engine/Source/Developer/SlateReflector/Private/Models/WidgetReflectorNode.cpp
--- a/Engine/Source/Developer/SlateReflector/Private/Models/WidgetReflectorNode.cpp
+++ b/Engine/Source/Developer/SlateReflector/Private/Models/WidgetReflectorNode.cpp
@@ -547,6 +547,8 @@ TSharedRef<FJsonValue> FSnapshotWidgetReflectorNode::ToJson(const TSharedRef<FSn
 	RootJsonObject->SetStringField(TEXT("WidgetVisibilityText"), RootSnapshotNode->CachedWidgetVisibilityText.ToString());
 	RootJsonObject->SetBoolField(TEXT("WidgetVisible"), RootSnapshotNode->bCachedWidgetVisible);
 	RootJsonObject->SetBoolField(TEXT("WidgetVisibleInherited"), RootSnapshotNode->bCachedWidgetVisibleInherited);
+	RootJsonObject->SetBoolField(TEXT("SupportsInvalidation"), RootSnapshotNode->bCachedSupportsInvalidation);
+	RootJsonObject->SetBoolField(TEXT("SupportsInvalidationRecursive"), RootSnapshotNode->bCachedSupportsInvalidationRecursive);
 	RootJsonObject->SetBoolField(TEXT("WidgetFocusable"), RootSnapshotNode->bCachedWidgetFocusable);
 	RootJsonObject->SetBoolField(TEXT("WidgetNeedsTick"), RootSnapshotNode->bCachedWidgetNeedsTick);
 	RootJsonObject->SetBoolField(TEXT("WidgetIsVolatile"), RootSnapshotNode->bCachedWidgetIsVolatile);
@@ -721,6 +723,8 @@ TSharedRef<FSnapshotWidgetReflectorNode> FSnapshotWidgetReflectorNode::FromJson(
 	RootSnapshotNode->CachedWidgetVisibilityText = FText::FromString(RootJsonObject->GetStringField(TEXT("WidgetVisibilityText")));
 	RootSnapshotNode->bCachedWidgetVisible = RootJsonObject->GetBoolField(TEXT("WidgetVisible"));
 	RootSnapshotNode->bCachedWidgetVisibleInherited = RootJsonObject->GetBoolField(TEXT("WidgetVisibleInherited"));
+	RootSnapshotNode->bCachedSupportsInvalidation = RootJsonObject->GetBoolField(TEXT("SupportsInvalidation"));
+	RootSnapshotNode->bCachedSupportsInvalidationRecursive = RootJsonObject->GetBoolField(TEXT("SupportsInvalidationRecursive"));
 	RootSnapshotNode->bCachedWidgetFocusable = RootJsonObject->GetBoolField(TEXT("WidgetFocusable"));
 	RootSnapshotNode->bCachedWidgetNeedsTick = RootJsonObject->GetBoolField(TEXT("WidgetNeedsTick"));
 	RootSnapshotNode->bCachedWidgetIsVolatile = RootJsonObject->GetBoolField(TEXT("WidgetIsVolatile"));
diff --git a/Engine/Source/Runtime/Slate/Private/Widgets/SInvalidationPanel.cpp b/Engine/Source/Runtime/Slate/Private/Widgets/SInvalidationPanel.cpp
--- a/Engine/Source/Runtime/Slate/Private/Widgets/SInvalidationPanel.cpp
+++ b/Engine/Source/Runtime/Slate/Private/Widgets/SInvalidationPanel.cpp
@@ -92,7 +92,7 @@ void SInvalidationPanel::Construct( const FArguments& InArgs )
 
 FChildren* SInvalidationPanel::GetAllChildren()
 {
-	// The default implementation would call SInvalidationPanel::GetChildren would is always empty if we can cache
+	// The default implementation would call SInvalidationPanel::GetChildren which is always empty if we can cache
 	return SCompoundWidget::GetChildren();
 }
 
diff --git a/Engine/Source/Runtime/Slate/Public/Widgets/SInvalidationPanel.h b/Engine/Source/Runtime/Slate/Public/Widgets/SInvalidationPanel.h
--- a/Engine/Source/Runtime/Slate/Public/Widgets/SInvalidationPanel.h
+++ b/Engine/Source/Runtime/Slate/Public/Widgets/SInvalidationPanel.h
@@ -44,7 +44,7 @@ public:
 	SLATE_API virtual FChildren* GetAllChildren() override;
 
 	/** Returns true when the invalidation panel is able to cache the widgets. **/
-	SLATE_API virtual bool GetCanCache() const;
+	SLATE_API bool GetCanCache() const;
 
 	SLATE_API void SetCanCache(bool InCanCache);
 
diff --git a/Engine/Source/Runtime/SlateCore/Public/Widgets/SWidget.h b/Engine/Source/Runtime/SlateCore/Public/Widgets/SWidget.h
--- a/Engine/Source/Runtime/SlateCore/Public/Widgets/SWidget.h
+++ b/Engine/Source/Runtime/SlateCore/Public/Widgets/SWidget.h
@@ -1523,7 +1523,7 @@ public:
 
 	/**
 	 * Returns true if this widget and all of its children support invalidation.
-	 * This value is cached and reused unless bUseCachedValue is set to true.
+	 * This value is cached and reused unless bUseCachedValue is set to false.
 	 */
 	SLATECORE_API virtual bool SupportsInvalidationRecursive(bool bUseCachedValue = true) const;
 private:
diff --git a/Engine/Source/Runtime/UMG/Private/Components/InvalidationBox.cpp b/Engine/Source/Runtime/UMG/Private/Components/InvalidationBox.cpp
--- a/Engine/Source/Runtime/UMG/Private/Components/InvalidationBox.cpp
+++ b/Engine/Source/Runtime/UMG/Private/Components/InvalidationBox.cpp
@@ -44,10 +44,6 @@ PRAGMA_ENABLE_DEPRECATION_WARNINGS
 		MyInvalidationPanel->SetContent(GetContentSlot()->Content ? GetContentSlot()->Content->TakeWidget() : SNullWidget::NullWidget);
 	}
 	
-#if UE_SLATE_WITH_DYNAMIC_INVALIDATION
-	MyInvalidationPanel->SetUseDynamicInvalidation(IsDesignTime() ? false : bUseDynamicInvalidation);
-#endif
-
 	return MyInvalidationPanel.ToSharedRef();
 }
 
@@ -94,31 +90,6 @@ void UInvalidationBox::SetCanCache(bool CanCache)
 }
 PRAGMA_ENABLE_DEPRECATION_WARNINGS
 
-bool UInvalidationBox::GetUseDynamicInvalidation() const
-{
-#if UE_SLATE_WITH_DYNAMIC_INVALIDATION
-	if ( MyInvalidationPanel.IsValid() )
-	{
-		return MyInvalidationPanel->GetUseDynamicInvalidation();
-	}
-	return bUseDynamicInvalidation;
-#else
-	return false;
-#endif
-}
-
-void UInvalidationBox::SetUseDynamicInvalidation(bool UseDynamicInvalidation)
-{
-#if UE_SLATE_WITH_DYNAMIC_INVALIDATION
-	bUseDynamicInvalidation = UseDynamicInvalidation;
-	if ( MyInvalidationPanel.IsValid() )
-	{
-		MyInvalidationPanel->SetUseDynamicInvalidation(UseDynamicInvalidation);
-	}
-#endif
-}
-
-
 #if WITH_EDITOR
 
 const FText UInvalidationBox::GetPaletteCategory()
diff --git a/Engine/Source/Runtime/UMG/Public/Components/InvalidationBox.h b/Engine/Source/Runtime/UMG/Public/Components/InvalidationBox.h
--- a/Engine/Source/Runtime/UMG/Public/Components/InvalidationBox.h
+++ b/Engine/Source/Runtime/UMG/Public/Components/InvalidationBox.h
@@ -43,19 +43,6 @@ public:
 	UFUNCTION(BlueprintCallable, Category="Invalidation Box")
 	UMG_API void SetCanCache(bool CanCache);
 
-	/**
-	 * @returns true when the invalidation box uses Dynamic Invalidation.
-	 */
-	UFUNCTION(BlueprintCallable, Category="Invalidation Box")
-	UMG_API bool GetUseDynamicInvalidation() const;
-
-	/**
-	 * Tell the InvalidationBox to use Dynamic Invalidation. This option is Editor only.
-	 * Dynamic Invalidation means that the widget's content will only be cached if all descendant widgets support invalidation.
-	 */
-	UFUNCTION(BlueprintCallable, Category="Invalidation Box")
-	UMG_API void SetUseDynamicInvalidation(bool UseDynamicInvalidation);
-
 	UMG_API virtual void ReleaseSlateResources(bool bReleaseChildren) override;
 
 #if WITH_EDITOR
@@ -82,13 +69,6 @@ protected:
 	UPROPERTY(EditAnywhere, BlueprintReadWrite, Getter = "GetCanCache", Setter = "SetCanCache", BlueprintGetter = "GetCanCache", BlueprintSetter = "SetCanCache", Category = "Caching")
 	bool bCanCache;
 
-	/**
-	 * Should Dynamic Invalidation be enabled for this Invalidation Panel?
-	 * Dynamic Invalidation means that the widget's content will only be cached if all descendant widgets support invalidation.
-	 */
-	UPROPERTY(EditAnywhere, BlueprintReadWrite, Getter = "GetUseDynamicInvalidation", Setter = "SetUseDynamicInvalidation", BlueprintGetter = "GetUseDynamicInvalidation", BlueprintSetter = "SetUseDynamicInvalidation", Category = "Caching")
-	bool bUseDynamicInvalidation;
-
 protected:
 	TSharedPtr<class SInvalidationPanel> MyInvalidationPanel;
 };
