WidgetReflector: Guard nullptr dereferences.

#rnx

diff --git a/Engine/Source/Developer/SlateReflector/Private/Widgets/NavigationSimulationOverlay.cpp b/Engine/Source/Developer/SlateReflector/Private/Widgets/NavigationSimulationOverlay.cpp
--- a/Engine/Source/Developer/SlateReflector/Private/Widgets/NavigationSimulationOverlay.cpp
+++ b/Engine/Source/Developer/SlateReflector/Private/Widgets/NavigationSimulationOverlay.cpp
@@ -272,11 +272,18 @@ int32 FNavigationSimulationOverlay::PaintSnapshotNode(const TArray<FNavigationSi
 	const NavigationSimulationOverlay::FOverlayBrushes OverlayBrushes;
 	for (const FNavigationSimulationWidgetNodePtr& NodePtr : NodeToPaint)
 	{
-		auto GetGeometry = [&GeometryMap, &AllottedGeometry, RootDrawOffset](const FNavigationSimulationWidgetInfo& Node) -> TOptional<FPaintGeometry>
+		auto GetGeometry = [&GeometryMap, &AllottedGeometry, RootDrawOffset](const FNavigationSimulationWidgetInfo& Info) -> TOptional<FPaintGeometry>
 		{
-			if (Node.bHasGeometry)
+			if (Info.bHasGeometry)
 			{
-				return TOptional<FPaintGeometry>(*GeometryMap.GetGeometry(Node, AllottedGeometry, RootDrawOffset));
+				if (const FPaintGeometry* PaintGeometry = GeometryMap.GetGeometry(Info, AllottedGeometry, RootDrawOffset))
+				{
+					return TOptional<FPaintGeometry>(*PaintGeometry);
+				}
+				else
+				{
+					UE_LOG(LogSlate, Warning, TEXT("Slate Navigation Simulation missing paint geometry for snapshot widget %s"), *(Info.WidgetTypeAndShortName.ToString()));
+				}
 			}
 			return TOptional<FPaintGeometry>();
 		};
@@ -299,7 +306,14 @@ int32 FNavigationSimulationOverlay::PaintLiveNode(const TArray<FNavigationSimula
 		{
 			if (TSharedPtr<const SWidget> WidgetLive = Info.WidgetLive.Pin())
 			{
-				return TOptional<FPaintGeometry>(*GeometryMap.GetGeometry(WidgetLive.ToSharedRef()));
+				if (const FPaintGeometry* PaintGeometry = GeometryMap.GetGeometry(WidgetLive.ToSharedRef()))
+				{
+					return TOptional<FPaintGeometry>(*PaintGeometry);
+				}
+				else
+				{
+					UE_LOG(LogSlate, Warning, TEXT("Slate Navigation Simulation missing paint geometry for live widget %s"), *(Info.WidgetTypeAndShortName.ToString()));
+				}
 			}
 			return TOptional<FPaintGeometry>();
 		};
