Fixed warnings related to widget reflector UX

#rnx

diff --git a/Engine/Source/Developer/SlateReflector/Private/Widgets/SNavigationSimulationList.cpp b/Engine/Source/Developer/SlateReflector/Private/Widgets/SNavigationSimulationList.cpp
--- a/Engine/Source/Developer/SlateReflector/Private/Widgets/SNavigationSimulationList.cpp
+++ b/Engine/Source/Developer/SlateReflector/Private/Widgets/SNavigationSimulationList.cpp
@@ -242,7 +242,7 @@ void SNavigationSimulationListBase::Construct(const FArguments& Args, ENavigatio
 	[
 		SNew(SSplitter)
 		.Orientation(EOrientation::Orient_Vertical)
-		.Style(FWidgetReflectorStyle::Get(), "SplitterDark")
+		.Style(FAppStyle::Get(), "SplitterDark")
 
 		+ SSplitter::Slot()
 		.Value(0.5f)
@@ -537,7 +537,6 @@ TSharedRef<SWidget> SDetailView::BuildNaviationSlots()
 TSharedRef<SWidget> SDetailView::BuildNaviationSlot(EUINavigation Navigation)
 {
 	return SNew(SCheckBox)
-	.Style(FWidgetReflectorStyle::Get(), "CheckBox")
 	.ForegroundColor(FSlateColor::UseForeground())
 	.IsChecked(this, &SDetailView::GetNavigationSlotChecked, Navigation)
 	.OnCheckStateChanged(this, &SDetailView::HandleNavigationSlotCheckStateChanged, Navigation)
