Slate Reflector: Fixed build issue 1243735

#Synced-CL 51353830

diff --git a/Engine/Source/Developer/SlateReflector/Private/Widgets/SWidgetReflectorTreeWidgetItem.cpp b/Engine/Source/Developer/SlateReflector/Private/Widgets/SWidgetReflectorTreeWidgetItem.cpp
--- a/Engine/Source/Developer/SlateReflector/Private/Widgets/SWidgetReflectorTreeWidgetItem.cpp
+++ b/Engine/Source/Developer/SlateReflector/Private/Widgets/SWidgetReflectorTreeWidgetItem.cpp
@@ -1,7 +1,6 @@
 // Copyright Epic Games, Inc. All Rights Reserved.
 
 #include "Widgets/SWidgetReflectorTreeWidgetItem.h"
-#include "Blueprint/UserWidget.h"
 #include "SlateOptMacros.h"
 #include "Widgets/Colors/SColorBlock.h"
 #include "Widgets/Input/SHyperlink.h"
@@ -9,6 +8,10 @@
 #include "Widgets/SInvalidationPanel.h"
 #include "Widgets/Input/SCheckBox.h"
 
+#if WITH_EDITOR
+#include "Blueprint/UserWidget.h"
+#endif
+
 #define LOCTEXT_NAMESPACE "SWidgetReflector"
 
 /* SMultiColumnTableRow overrides
diff --git a/Engine/Source/Developer/SlateReflector/SlateReflector.Build.cs b/Engine/Source/Developer/SlateReflector/SlateReflector.Build.cs
--- a/Engine/Source/Developer/SlateReflector/SlateReflector.Build.cs
+++ b/Engine/Source/Developer/SlateReflector/SlateReflector.Build.cs
@@ -27,7 +27,6 @@ public class SlateReflector : ModuleRules
 				"DesktopPlatform",
 				"ImageCore",
 				"ImageWrapper",
-				"UMG"
             }
 		);
 
@@ -53,7 +52,8 @@ public class SlateReflector : ModuleRules
 				new string[] {
 					"PropertyEditor",
 					"UnrealEd",
-                }
+					"UMG",
+				}
 			);
 
             PrivateIncludePathModuleNames.AddRange(
