Widget Reflector: Updated to source location for UWidgets.
- The widget class a widget is defined in as well as the asset it is contained in are both listed in the tooltip and the source location.
- When clicking the widget class it will open and select the debug instance of the widget.
- When clicking the asset/var name it will open the asset.

#Synced-CL 51657968

diff --git a/Engine/Source/Developer/SlateReflector/Private/Widgets/SWidgetReflectorToolTipWidget.cpp b/Engine/Source/Developer/SlateReflector/Private/Widgets/SWidgetReflectorToolTipWidget.cpp
--- a/Engine/Source/Developer/SlateReflector/Private/Widgets/SWidgetReflectorToolTipWidget.cpp
+++ b/Engine/Source/Developer/SlateReflector/Private/Widgets/SWidgetReflectorToolTipWidget.cpp
@@ -10,6 +10,10 @@
 #include "Widgets/Layout/SGridPanel.h"
 #include "Widgets/Text/STextBlock.h"
 
+#if WITH_EDITOR
+#include "Blueprint/UserWidget.h"
+#endif
+
 #define LOCTEXT_NAMESPACE "SReflectorToolTipWidget"
 
 const FText SReflectorToolTipWidget::TrueText = LOCTEXT("True", "True");
@@ -129,6 +133,24 @@ void SReflectorToolTipWidget::Construct( const FArguments& InArgs )
 	}
 }
 
+FText SReflectorToolTipWidget::GetWidgetsSourceName() const
+{
+#if WITH_EDITOR
+	if (UUserWidget* SourceObject = Cast<UUserWidget>(WidgetInfo->GetWidgetSourceObject()))
+	{
+		UClass* SourceClass = SourceObject->GetClass();
+
+		return FText::Format(
+			LOCTEXT("WidgetTypeInAsset", "{0} in {1}"),
+			FText::FromString(SourceClass->ClassGeneratedBy->GetName()),
+			WidgetInfo->GetWidgetReadableLocation()
+		);
+	}
+#endif
+
+	return WidgetInfo->GetWidgetReadableLocation();
+}
+
 FText SReflectorToolTipWidget::GetLiveAttributeName() const
 {
 	if (TSharedPtr<SWidget> Widget = WidgetInfo->GetLiveWidget())
diff --git a/Engine/Source/Developer/SlateReflector/Private/Widgets/SWidgetReflectorToolTipWidget.h b/Engine/Source/Developer/SlateReflector/Private/Widgets/SWidgetReflectorToolTipWidget.h
--- a/Engine/Source/Developer/SlateReflector/Private/Widgets/SWidgetReflectorToolTipWidget.h
+++ b/Engine/Source/Developer/SlateReflector/Private/Widgets/SWidgetReflectorToolTipWidget.h
@@ -29,10 +29,7 @@ public:
 
 private:
 
-	FText GetWidgetsSourceName() const
-	{
-		return WidgetInfo->GetWidgetReadableLocation();
-	}
+	FText GetWidgetsSourceName() const;
 
 	FText GetWidgetsDesiredSize() const
 	{
diff --git a/Engine/Source/Developer/SlateReflector/Private/Widgets/SWidgetReflectorTreeWidgetItem.cpp b/Engine/Source/Developer/SlateReflector/Private/Widgets/SWidgetReflectorTreeWidgetItem.cpp
--- a/Engine/Source/Developer/SlateReflector/Private/Widgets/SWidgetReflectorTreeWidgetItem.cpp
+++ b/Engine/Source/Developer/SlateReflector/Private/Widgets/SWidgetReflectorTreeWidgetItem.cpp
@@ -197,15 +197,47 @@ TSharedRef<SWidget> SReflectorTreeWidgetItem::GenerateWidgetForColumn(const FNam
 	}
 	else if (ColumnName == NAME_WidgetInfo )
 	{
-		return SNew(SBox)
+		TSharedRef<SHorizontalBox> Info = SNew(SHorizontalBox);
+
+#if WITH_EDITOR
+		if (UUserWidget* SourceObject = Cast<UUserWidget>(WidgetInfo->GetWidgetSourceObject()))
+		{
+			UClass* SourceClass = SourceObject->GetClass();
+
+			Info->AddSlot()
+				.AutoWidth()
+				.HAlign(HAlign_Left)
+				.VAlign(VAlign_Center)
+				.Padding(FMargin(2.0f, 0.0f))
+				[
+					SNew(SHyperlink)
+					.Text(SourceClass->GetDisplayNameText())
+					.OnNavigate(this, &SReflectorTreeWidgetItem::HandleHyperlinkNavigate, /* Attempt class debug */ true)
+				];
+
+			Info->AddSlot()
+				.AutoWidth()
+				.HAlign(HAlign_Left)
+				.VAlign(VAlign_Center)
+				[
+					SNew(STextBlock)
+					.Text(LOCTEXT("Source_In", "in"))
+				];
+		}
+#endif
+			
+		Info->AddSlot()
+			.AutoWidth()
 			.HAlign(HAlign_Left)
 			.VAlign(VAlign_Center)
 			.Padding(FMargin(2.0f, 0.0f))
 			[
 				SNew(SHyperlink)
 				.Text(WidgetInfo->GetWidgetReadableLocation())
-				.OnNavigate(this, &SReflectorTreeWidgetItem::HandleHyperlinkNavigate)
+				.OnNavigate(this, &SReflectorTreeWidgetItem::HandleHyperlinkNavigate, /* Attempt class debug */ false)
 			];
+
+		return Info;
 	}
 	else if (ColumnName == NAME_Visibility )
 	{
@@ -315,23 +347,26 @@ TSharedRef<SWidget> SReflectorTreeWidgetItem::GenerateWidgetForColumn(const FNam
 }
 END_SLATE_FUNCTION_BUILD_OPTIMIZATION
 
-void SReflectorTreeWidgetItem::HandleHyperlinkNavigate()
+void SReflectorTreeWidgetItem::HandleHyperlinkNavigate(bool bAttemptClassDebug)
 {
 #if WITH_EDITOR
 	bool bHasOpenedAsset = false;
 
-	if (UUserWidget* SourceObject = Cast<UUserWidget>(WidgetInfo->GetWidgetSourceObject()))
+	if (bAttemptClassDebug)
 	{
-		UClass* SourceClass = SourceObject->GetClass();
-
-		if (SourceClass->ClassGeneratedBy)
+		if (UUserWidget* SourceObject = Cast<UUserWidget>(WidgetInfo->GetWidgetSourceObject()))
 		{
-			OnAccessAsset.Execute(SourceClass->ClassGeneratedBy);
-			bHasOpenedAsset = true;
+			UClass* SourceClass = SourceObject->GetClass();
 
-			if (OnAccessDebugObject.IsBound())
+			if (SourceClass->ClassGeneratedBy)
 			{
-				OnAccessDebugObject.Execute(SourceObject);
+				OnAccessAsset.Execute(SourceClass->ClassGeneratedBy);
+				bHasOpenedAsset = true;
+
+				if (OnAccessDebugObject.IsBound())
+				{
+					OnAccessDebugObject.Execute(SourceObject);
+				}
 			}
 		}
 	}
diff --git a/Engine/Source/Developer/SlateReflector/Private/Widgets/SWidgetReflectorTreeWidgetItem.h b/Engine/Source/Developer/SlateReflector/Private/Widgets/SWidgetReflectorTreeWidgetItem.h
--- a/Engine/Source/Developer/SlateReflector/Private/Widgets/SWidgetReflectorTreeWidgetItem.h
+++ b/Engine/Source/Developer/SlateReflector/Private/Widgets/SWidgetReflectorTreeWidgetItem.h
@@ -67,7 +67,7 @@ protected:
 		return WidgetInfo->GetTint();
 	}
 
-	void HandleHyperlinkNavigate();
+	void HandleHyperlinkNavigate(bool bAttemptClassDebug);
 
 private:
 
