[SlateReflector] Add Draw Elements tab to Widget Reflector Display the captured draw elements by FVisualTreeCapture. It gives a per-frame view of every Slate  draw element alongside its widget owner, element type, and layer ID.
#Synced-CL 51824547

diff --git a/Engine/Source/Developer/SlateReflector/Private/VisualTreeCapture.cpp b/Engine/Source/Developer/SlateReflector/Private/VisualTreeCapture.cpp
--- a/Engine/Source/Developer/SlateReflector/Private/VisualTreeCapture.cpp
+++ b/Engine/Source/Developer/SlateReflector/Private/VisualTreeCapture.cpp
@@ -90,12 +90,12 @@ bool FVisualEntry::IsPointInside(const FVector2f& Point) const
 	return false;
 }
 
-TSharedPtr<const SWidget> FVisualTreeSnapshot::Pick(FVector2f Point)
+TSharedPtr<const SWidget> FVisualTreeSnapshot::Pick(FVector2f Point) const
 {
-	for (int Index = Entries.Num() - 1; Index >= 0; Index--)
+	for (int32 Index = Entries.Num() - 1; Index >= 0; Index--)
 	{
 		const FVisualEntry& Entry = Entries[Index];
-		if (Entry.ClippingIndex != -1)
+		if (Entry.ClippingIndex != INDEX_NONE)
 		{
 			const TArray<FSlateClippingState>& LocalClippingState = Entry.bFromCache ? CachedClippingStates : ClippingStates;
 			if (ensure(LocalClippingState.IsValidIndex(Entry.ClippingIndex)))
@@ -118,6 +118,35 @@ TSharedPtr<const SWidget> FVisualTreeSnapshot::Pick(FVector2f Point)
 	return TSharedPtr<const SWidget>();
 }
 
+TArray<int32> FVisualTreeSnapshot::PickAll(FVector2f Point) const
+{
+	TArray<int32> Result;
+	for (int Index = Entries.Num() - 1; Index >= 0; Index--)
+	{
+		const FVisualEntry& Entry = Entries[Index];
+		if (Entry.IsPointInside(Point))
+		{
+			if (Entry.ClippingIndex != INDEX_NONE)
+			{
+				const TArray<FSlateClippingState>& LocalClippingState = Entry.bFromCache ? CachedClippingStates : ClippingStates;
+				if (ensure(LocalClippingState.IsValidIndex(Entry.ClippingIndex)))
+				{
+					if (LocalClippingState[Entry.ClippingIndex].IsPointInside(Point))
+					{
+						Result.Add(Index);
+					}
+				}
+			}
+			else
+			{
+				Result.Add(Index);
+			}
+		}
+	}
+
+	return Result;
+}
+
 FVisualTreeCapture::FVisualTreeCapture()
 	: bIsEnabled(false)
 	, WindowIsInvalidationRootCounter(0)
@@ -175,7 +204,7 @@ void FVisualTreeCapture::Reset()
 	VisualTrees.Reset();
 }
 
-TSharedPtr<FVisualTreeSnapshot> FVisualTreeCapture::GetVisualTreeForWindow(SWindow* InWindow)
+TSharedPtr<const FVisualTreeSnapshot> FVisualTreeCapture::GetVisualTreeForWindow(SWindow* InWindow) const
 {
 	return VisualTrees.FindRef(InWindow);
 }
diff --git a/Engine/Source/Developer/SlateReflector/Private/VisualTreeCapture.h b/Engine/Source/Developer/SlateReflector/Private/VisualTreeCapture.h
--- a/Engine/Source/Developer/SlateReflector/Private/VisualTreeCapture.h
+++ b/Engine/Source/Developer/SlateReflector/Private/VisualTreeCapture.h
@@ -3,6 +3,7 @@
 #pragma once
 
 #include "CoreMinimal.h"
+#include "Rendering/DrawElementCoreTypes.h"
 
 class SWindow;
 class SWidget;
@@ -14,7 +15,6 @@ struct FGeometry;
 class FSlateRect;
 class FSlateInvalidationRoot;
 struct FSlateDebuggingElementTypeAddedEventArgs;
-enum class EElementType : uint8;
 
 class FVisualEntry
 {
@@ -24,11 +24,11 @@ public:
 	FVector2f BottomLeft;
 	FVector2f BottomRight;
 
-	int32 LayerId;
-	int32 ClippingIndex;
-	int32 ElementIndex;
-	EElementType ElementType;
-	bool bFromCache;
+	int32 LayerId = INDEX_NONE;
+	int32 ClippingIndex = INDEX_NONE;
+	int32 ElementIndex = INDEX_NONE;
+	EElementType ElementType = EElementType::ET_Count;
+	bool bFromCache = false;
 	TWeakPtr<const SWidget> Widget;
 
 	FVisualEntry(const TWeakPtr<const SWidget>& Widget, int32 InElementIndex, EElementType InElementType);
@@ -42,7 +42,9 @@ public:
 class FVisualTreeSnapshot : public TSharedFromThis<FVisualTreeSnapshot>
 {
 public:
-	TSharedPtr<const SWidget> Pick(FVector2f Point);
+	TSharedPtr<const SWidget> Pick(FVector2f Point) const;
+	/** Returns all the the Entires at the location. */
+	TArray<int32> PickAll(FVector2f Point) const;
 	
 public:
 	TArray<FVisualEntry> Entries;
@@ -66,7 +68,19 @@ public:
 	/** Resets the visual tree capture to a pre-capture state and destroys the cached visual tree captured last. */
 	void Reset();
 
-	TSharedPtr<FVisualTreeSnapshot> GetVisualTreeForWindow(SWindow* InWindow);
+	/** The visual tree capture is enabled. */
+	bool IsEnabled() const
+	{
+		return bIsEnabled;
+	}
+
+	/** The visual tree capture contains elements. */
+	bool HasElements() const
+	{
+		return VisualTrees.Num() > 0;
+	}
+
+	TSharedPtr<const FVisualTreeSnapshot> GetVisualTreeForWindow(SWindow* InWindow) const;
 	
 private:
 
@@ -85,6 +99,7 @@ private:
 	void ElementTypeAdded(const FSlateDebuggingElementTypeAddedEventArgs& ElementTypeAddedArgs);
 
 	void OnWindowBeingDestroyed(const SWindow& WindowBeingDestoyed);
+
 private:
 	TMap<const SWindow*, TSharedPtr<FVisualTreeSnapshot>> VisualTrees;
 	bool bIsEnabled;
diff --git a/Engine/Source/Developer/SlateReflector/Private/Widgets/SVisualElementsTab.cpp b/Engine/Source/Developer/SlateReflector/Private/Widgets/SVisualElementsTab.cpp
new file mode 100644
--- /dev/null
+++ b/Engine/Source/Developer/SlateReflector/Private/Widgets/SVisualElementsTab.cpp
@@ -0,0 +1,330 @@
+﻿// Copyright Epic Games, Inc. All Rights Reserved.
+
+#include "Widgets/SVisualElementsTab.h"
+
+#include "Framework/Application/SlateApplication.h"
+#include "Models/WidgetReflectorNode.h"
+#include "VisualTreeCapture.h"
+#include "Widgets/Input/SCheckBox.h"
+#include "Widgets/Input/SHyperlink.h"
+#include "Widgets/Layout/SBorder.h"
+#include "Widgets/Layout/SBox.h"
+#include "Widgets/Layout/SSplitter.h"
+#include "Widgets/SNullWidget.h"
+#include "Widgets/SWindow.h"
+#include "Widgets/Text/STextBlock.h"
+#include "Widgets/Views/SHeaderRow.h"
+#include "Widgets/Views/STableRow.h"
+
+#define LOCTEXT_NAMESPACE "SVisualElementsTab"
+
+/* FVisualElementEntry
+ *****************************************************************************/
+
+UE::SlateReflector::Private::FVisualElementEntry::FVisualElementEntry(const FVisualEntry& InEntry)
+	: Widget(InEntry.Widget)
+	, LayerId(InEntry.LayerId)
+	, ElementType(InEntry.ElementType)
+{
+}
+
+/* Column names
+ *****************************************************************************/
+
+namespace VisualElementsTabColumns
+{
+	static const FName Widget = "Widget";
+	static const FName Type   = "Type";
+	static const FName Layer  = "Layer";
+}
+
+/* Helpers
+ *****************************************************************************/
+
+static FText GetElementTypeName(EElementType Type)
+{
+	switch (Type)
+	{
+	case EElementType::ET_Box:             return LOCTEXT("ET_Box",           "Box");
+	case EElementType::ET_DebugQuad:       return LOCTEXT("ET_DebugQuad",     "DebugQuad");
+	case EElementType::ET_Text:            return LOCTEXT("ET_Text",           "Text");
+	case EElementType::ET_ShapedText:      return LOCTEXT("ET_ShapedText",     "ShapedText");
+	case EElementType::ET_Spline:          return LOCTEXT("ET_Spline",         "Spline");
+	case EElementType::ET_Line:            return LOCTEXT("ET_Line",           "Line");
+	case EElementType::ET_Gradient:        return LOCTEXT("ET_Gradient",       "Gradient");
+	case EElementType::ET_Viewport:        return LOCTEXT("ET_Viewport",       "Viewport");
+	case EElementType::ET_Border:          return LOCTEXT("ET_Border",         "Border");
+	case EElementType::ET_Custom:          return LOCTEXT("ET_Custom",         "Custom");
+	case EElementType::ET_CustomVerts:     return LOCTEXT("ET_CustomVerts",    "CustomVerts");
+	case EElementType::ET_PostProcessPass: return LOCTEXT("ET_PostProcess",    "PostProcess");
+	case EElementType::ET_RoundedBox:      return LOCTEXT("ET_RoundedBox",     "RoundedBox");
+	default:                               return LOCTEXT("ET_Unknown",        "Unknown");
+	}
+}
+
+/* SVisualEntryTableRow
+ *****************************************************************************/
+
+using FVisualElementEntry = UE::SlateReflector::Private::FVisualElementEntry;
+
+class SVisualEntryTableRow : public SMultiColumnTableRow<TSharedPtr<FVisualElementEntry>>
+{
+public:
+	SLATE_BEGIN_ARGS(SVisualEntryTableRow) {}
+		SLATE_EVENT(SVisualElementsTab::FOnWidgetSelected, OnWidgetSelected)
+	SLATE_END_ARGS()
+
+	void Construct(const FArguments& InArgs, TSharedPtr<FVisualElementEntry> InEntry, const TSharedRef<STableViewBase>& OwnerTable)
+	{
+		Entry = InEntry;
+		OnWidgetSelected = InArgs._OnWidgetSelected;
+		FSuperRowType::Construct(FSuperRowType::FArguments(), OwnerTable);
+	}
+
+	virtual TSharedRef<SWidget> GenerateWidgetForColumn(const FName& ColumnName) override
+	{
+		if (ColumnName == VisualElementsTabColumns::Widget)
+		{
+			FText WidgetName = LOCTEXT("InvalidWidget", "(invalid)");
+			if (TSharedPtr<const SWidget> Widget = Entry->Widget.Pin())
+			{
+				WidgetName = FWidgetReflectorNodeUtils::GetWidgetTypeAndShortName(Widget);
+			}
+
+			return SNew(SBox)
+				.VAlign(VAlign_Center)
+				.Padding(FMargin(4.0f, 0.0f))
+				[
+					SNew(SHyperlink)
+					.Text(WidgetName)
+					.OnNavigate(this, &SVisualEntryTableRow::HandleNavigate)
+				];
+		}
+
+		if (ColumnName == VisualElementsTabColumns::Type)
+		{
+			return SNew(SBox)
+				.VAlign(VAlign_Center)
+				.Padding(FMargin(4.0f, 0.0f))
+				[
+					SNew(STextBlock)
+					.Text(GetElementTypeName(Entry->ElementType))
+				];
+		}
+
+		if (ColumnName == VisualElementsTabColumns::Layer)
+		{
+			return SNew(SBox)
+				.VAlign(VAlign_Center)
+				.Padding(FMargin(4.0f, 0.0f))
+				[
+					SNew(STextBlock)
+					.Text(FText::AsNumber(Entry->LayerId))
+				];
+		}
+
+		return SNullWidget::NullWidget;
+	}
+
+private:
+	void HandleNavigate()
+	{
+		OnWidgetSelected.ExecuteIfBound(Entry->Widget.Pin());
+	}
+
+	TSharedPtr<FVisualElementEntry> Entry;
+	SVisualElementsTab::FOnWidgetSelected OnWidgetSelected;
+};
+
+/* Shared header row builder
+ *****************************************************************************/
+
+static TSharedRef<SHeaderRow> MakeHeaderRow()
+{
+	return SNew(SHeaderRow)
+		+ SHeaderRow::Column(VisualElementsTabColumns::Widget)
+		.DefaultLabel(LOCTEXT("ColumnWidget", "Widget"))
+		.FillWidth(1.0f)
+		+ SHeaderRow::Column(VisualElementsTabColumns::Type)
+		.DefaultLabel(LOCTEXT("ColumnType", "Type"))
+		.FixedWidth(80.0f)
+		+ SHeaderRow::Column(VisualElementsTabColumns::Layer)
+		.DefaultLabel(LOCTEXT("ColumnLayer", "Layer"))
+		.FixedWidth(50.0f);
+}
+
+/* SVisualElementsTab
+ *****************************************************************************/
+
+void SVisualElementsTab::Construct(const FArguments& InArgs, TNotNull<const FVisualTreeCapture*> InVisualCapture)
+{
+	OnWidgetSelected = InArgs._OnWidgetSelected;
+	VisualCapture = InVisualCapture;
+	FilterWidgets = InArgs._FilterWidgets;
+	IsPickingActive = InArgs._IsPickingActive;
+
+	ChildSlot
+	[
+		SNew(SVerticalBox)
+
+		// Toolbar
+		+ SVerticalBox::Slot()
+		.AutoHeight()
+		.Padding(FMargin(2.0f, 2.0f))
+		[
+			SNew(SHorizontalBox)
+			+ SHorizontalBox::Slot()
+			.AutoWidth()
+			.VAlign(VAlign_Center)
+			[
+				SNew(SCheckBox)
+				.IsChecked(this, &SVisualElementsTab::GetFilterByWidgetChecked)
+				.OnCheckStateChanged(this, &SVisualElementsTab::HandleFilterByWidgetChanged)
+				.ToolTipText(LOCTEXT("FilterTooltip", "When checked, the main list only shows elements produced by the selected widgets in the tree."))
+				[
+					SNew(STextBlock)
+					.Text(LOCTEXT("FilterLabel", "Filter by Selected Widgets"))
+				]
+			]
+		]
+
+		+ SVerticalBox::Slot()
+		.FillHeight(1.0f)
+		[
+			SNew(SSplitter)
+			.Orientation(Orient_Vertical)
+
+			// Main entry list
+			+ SSplitter::Slot()
+			.Value(0.6f)
+			[
+				SAssignNew(ListView, SListView<TSharedPtr<FVisualElementEntry>>)
+				.ListItemsSource(&ListEntries)
+				.OnGenerateRow(this, &SVisualElementsTab::HandleGenerateRow)
+				.HeaderRow(MakeHeaderRow())
+			]
+
+			// Region entry list
+			+ SSplitter::Slot()
+			.Value(0.4f)
+			[
+				SNew(SVerticalBox)
+
+				+ SVerticalBox::Slot()
+				.AutoHeight()
+				[
+					SNew(SBorder)
+					.Padding(FMargin(4.0f, 2.0f))
+					[
+						SNew(STextBlock)
+						.Text(LOCTEXT("RegionHeader", "Region (entries inside cursor element)"))
+					]
+				]
+
+				+ SVerticalBox::Slot()
+				.FillHeight(1.0f)
+				[
+					SAssignNew(RegionListView, SListView<TSharedPtr<FVisualElementEntry>>)
+					.ListItemsSource(&RegionEntries)
+					.OnGenerateRow(this, &SVisualElementsTab::HandleGenerateRegionRow)
+					.HeaderRow(MakeHeaderRow())
+				]
+			]
+		]
+	];
+}
+
+void SVisualElementsTab::Tick(const FGeometry& AllottedGeometry, const double InCurrentTime, const float InDeltaTime)
+{
+	Super::Tick(AllottedGeometry, InCurrentTime, InDeltaTime);
+
+	check(VisualCapture);
+
+	const bool bHadElements = ListEntries.Num() > 0 || RegionEntries.Num() > 0;
+	const bool bIsPickingActive = IsPickingActive.Get(false);
+	ListEntries.Reset();
+	if (bIsPickingActive || !VisualCapture->HasElements())
+	{
+		RegionEntries.Reset();
+	}
+
+	const TArray<TSharedPtr<const SWidget>> PinnedFilters = FilterWidgets.Get();
+	const FVector2f ScreenCursorPos = FSlateApplication::Get().GetCursorPos();
+
+	TArray<TSharedRef<SWindow>> AllWindows;
+	FSlateApplication::Get().GetAllVisibleWindowsOrdered(AllWindows);
+
+	for (const TSharedRef<SWindow>& Window : AllWindows)
+	{
+		TSharedPtr<const FVisualTreeSnapshot> Snapshot = VisualCapture->GetVisualTreeForWindow(&Window.Get());
+		if (!Snapshot)
+		{
+			continue;
+		}
+
+		// Populate the main list (all windows).
+		for (const FVisualEntry& SourceEntry : Snapshot->Entries)
+		{
+			if (bFilterByWidget && !PinnedFilters.Contains(SourceEntry.Widget.Pin()))
+			{
+				continue;
+			}
+			ListEntries.Add(MakeShared<FVisualElementEntry>(SourceEntry));
+		}
+
+		// Populate the region list only while picking is active and for the window that contains the cursor.
+		// Use GetRectInScreen() to check containment in screen space, then convert
+		// the cursor to window-local space before passing to PickAll.
+		if (bIsPickingActive && RegionEntries.IsEmpty())
+		{
+			const FSlateRect WindowRect = Window->GetRectInScreen();
+			const bool bCursorInWindow = ScreenCursorPos.X >= WindowRect.Left
+				&& ScreenCursorPos.X <  WindowRect.Right
+				&& ScreenCursorPos.Y >= WindowRect.Top
+				&& ScreenCursorPos.Y <  WindowRect.Bottom;
+
+			if (bCursorInWindow)
+			{
+				// Convert to window-local space: subtract the window's screen position.
+				// GetLocalToWindowTransform() is local->window and contains only DPI scale.
+				const FVector2f WindowLocalCursor = ScreenCursorPos - FVector2f(Window->GetPositionInScreen());
+
+				for (int32 Index : Snapshot->PickAll(WindowLocalCursor))
+				{
+					RegionEntries.Add(MakeShared<FVisualElementEntry>(Snapshot->Entries[Index]));
+				}
+			}
+		}
+	}
+
+	const bool bHasElements = ListEntries.Num() > 0 || RegionEntries.Num() > 0;
+	if (bHadElements || bHasElements)
+	{
+		ListView->RequestListRefresh();
+		RegionListView->RequestListRefresh();
+	}
+}
+
+ECheckBoxState SVisualElementsTab::GetFilterByWidgetChecked() const
+{
+	return bFilterByWidget ? ECheckBoxState::Checked : ECheckBoxState::Unchecked;
+}
+
+void SVisualElementsTab::HandleFilterByWidgetChanged(ECheckBoxState NewState)
+{
+	bFilterByWidget = (NewState == ECheckBoxState::Checked);
+}
+
+TSharedRef<ITableRow> SVisualElementsTab::HandleGenerateRow(TSharedPtr<FVisualElementEntry> Entry, const TSharedRef<STableViewBase>& OwnerTable)
+{
+	return SNew(SVisualEntryTableRow, Entry, OwnerTable)
+		.OnWidgetSelected(OnWidgetSelected);
+}
+
+TSharedRef<ITableRow> SVisualElementsTab::HandleGenerateRegionRow(TSharedPtr<FVisualElementEntry> Entry, const TSharedRef<STableViewBase>& OwnerTable)
+{
+	return SNew(SVisualEntryTableRow, Entry, OwnerTable)
+		.OnWidgetSelected(OnWidgetSelected);
+}
+
+#undef LOCTEXT_NAMESPACE
diff --git a/Engine/Source/Developer/SlateReflector/Private/Widgets/SVisualElementsTab.h b/Engine/Source/Developer/SlateReflector/Private/Widgets/SVisualElementsTab.h
new file mode 100644
--- /dev/null
+++ b/Engine/Source/Developer/SlateReflector/Private/Widgets/SVisualElementsTab.h
@@ -0,0 +1,81 @@
+// Copyright Epic Games, Inc. All Rights Reserved.
+
+#pragma once
+
+#include "CoreMinimal.h"
+#include "Rendering/DrawElementCoreTypes.h"
+#include "Widgets/DeclarativeSyntaxSupport.h"
+#include "Widgets/SCompoundWidget.h"
+#include "Widgets/Views/SListView.h"
+
+class SWidget;
+class FVisualEntry;
+class FVisualTreeCapture;
+class FVisualTreeSnapshot;
+class ITableRow;
+class STableViewBase;
+
+namespace UE::SlateReflector::Private
+{
+
+/** Lightweight display entry for the Elements tab: widget pointer, layer, and element type. */
+struct FVisualElementEntry
+{
+	TWeakPtr<const SWidget> Widget;
+	int32 LayerId = INDEX_NONE;
+	EElementType ElementType = EElementType::ET_Count;
+
+	explicit FVisualElementEntry(const FVisualEntry& InEntry);
+};
+
+} // namespace UE::SlateReflector::Private
+
+/**
+ * Tab displaying all FVisualEntry items captured by FVisualTreeCapture.
+ */
+class SVisualElementsTab : public SCompoundWidget
+{
+private:
+	using Super = SCompoundWidget;
+
+	using FVisualElementEntry = UE::SlateReflector::Private::FVisualElementEntry;
+
+public:
+	DECLARE_DELEGATE_OneParam(FOnWidgetSelected, TSharedPtr<const SWidget>)
+
+	SLATE_BEGIN_ARGS(SVisualElementsTab) {}
+		SLATE_EVENT(FOnWidgetSelected, OnWidgetSelected)
+		SLATE_ATTRIBUTE(TArray<TSharedPtr<const SWidget>>, FilterWidgets)
+		SLATE_ATTRIBUTE(bool, IsPickingActive)
+	SLATE_END_ARGS()
+
+	void Construct(const FArguments& InArgs, TNotNull<const FVisualTreeCapture*> VisualCapture);
+
+	virtual void Tick(const FGeometry& AllottedGeometry, const double InCurrentTime, const float InDeltaTime) override;
+
+private:
+	ECheckBoxState GetFilterByWidgetChecked() const;
+	void HandleFilterByWidgetChanged(ECheckBoxState NewState);
+
+	TSharedRef<ITableRow> HandleGenerateRow(TSharedPtr<FVisualElementEntry> Entry, const TSharedRef<STableViewBase>& OwnerTable);
+	TSharedRef<ITableRow> HandleGenerateRegionRow(TSharedPtr<FVisualElementEntry> Entry, const TSharedRef<STableViewBase>& OwnerTable);
+
+	/** Returns true if all four corners of Entry are inside the Region's quad. */
+	static bool IsEntryInsideRegion(const FVisualEntry& Entry, const FVisualEntry& Region);
+
+	FOnWidgetSelected OnWidgetSelected;
+	const FVisualTreeCapture* VisualCapture = nullptr;
+	bool bFilterByWidget = false;
+	bool bDrawablePickingActive = false;
+
+	TAttribute<TArray<TSharedPtr<const SWidget>>> FilterWidgets;
+	TAttribute<bool> IsPickingActive;
+
+	// Main list (filtered subset of snapshot entries).
+	TSharedPtr<SListView<TSharedPtr<FVisualElementEntry>>> ListView;
+	TArray<TSharedPtr<FVisualElementEntry>> ListEntries;
+
+	// Region list (entries whose quads fit inside the entry under cursor).
+	TSharedPtr<SListView<TSharedPtr<FVisualElementEntry>>> RegionListView;
+	TArray<TSharedPtr<FVisualElementEntry>> RegionEntries;
+};
diff --git a/Engine/Source/Developer/SlateReflector/Private/Widgets/SWidgetReflector.cpp b/Engine/Source/Developer/SlateReflector/Private/Widgets/SWidgetReflector.cpp
--- a/Engine/Source/Developer/SlateReflector/Private/Widgets/SWidgetReflector.cpp
+++ b/Engine/Source/Developer/SlateReflector/Private/Widgets/SWidgetReflector.cpp
@@ -43,6 +43,7 @@
 #include "Widgets/Notifications/SNotificationList.h"
 #include "Widgets/Navigation/SBreadcrumbTrail.h"
 #include "Widgets/Input/SComboBox.h"
+#include "Widgets/SVisualElementsTab.h"
 #include "Widgets/SWidgetSnapshotVisualizer.h"
 #include "WidgetSnapshotService.h"
 #include "Types/ReflectionMetadata.h"
@@ -134,6 +135,7 @@ namespace WidgetReflectorTabID
 	static const FName WidgetEvents = "WidgetReflector.WidgetEventsTab";
 	static const FName HittestGrid = "WidgetReflector.HittestGridTab";
 	static const FName WidgetList = "WidgetReflector.WidgetList";
+	static const FName VisualElements = "WidgetReflector.VisualElements";
 }
 
 namespace WidgetReflectorText
@@ -156,6 +158,11 @@ namespace WidgetReflectorIcon
 	static const FName TakeSnapshot = "Icon.TakeSnapshot";
 }
 
+namespace VisualCaptureEnabled
+{
+	static const FName WidgetPickingMode = "WidgetPickingMode";
+}
+
 enum class EWidgetPickingMode : uint8
 {
 	None = 0,
@@ -207,6 +214,7 @@ private:
 #endif
 
 	TSharedRef<SDockTab> SpawnWidgetList(const FSpawnTabArgs& Args);
+	TSharedRef<SDockTab> SpawnVisualElementsTab(const FSpawnTabArgs& Args);
 
 public:
 
@@ -307,6 +315,12 @@ public:
 	/** Clear previous selection and set the selection to the live widget. */
 	void SelectLiveWidget( TSharedPtr<const SWidget> InWidget );
 
+	/** Get all selected live widgets. */
+	TArray<TSharedPtr<const SWidget>> GetSelectedLiveWidgets() const;
+
+	/** Returns true when Drawable picking mode is active. */
+	bool IsDrawablePickingActive() const { return PickingMode == EWidgetPickingMode::Drawable; }
+
 	/** Set the given nodes as the root of the tree. */
 	void SetNodesAsReflectorTreeRoot(TArray<TSharedRef<FWidgetReflectorNodeBase>> RootNodes);
 
@@ -334,25 +348,33 @@ public:
 	{
 		if (PickingMode != InMode)
 		{
-			// Disable visual picking, and re-enable widget caching.
-			VisualCapture.Disable();
+			// If leaving Drawable mode, remove its entry from the enablers.
+			if (PickingMode == EWidgetPickingMode::Drawable)
+			{
+				DisableVisualCapture(VisualCaptureEnabled::WidgetPickingMode);
+			}
 
 			// Enable the picking mode.
 			PickingMode = InMode;
 
-			// If we're enabling hit test, reset the visual capture entirely, we don't want to use the visual tree.
-			if (PickingMode == EWidgetPickingMode::HitTesting)
+			// If we're entering Drawable mode, add its entry to the enablers.
+			if (PickingMode == EWidgetPickingMode::Drawable)
 			{
-				VisualCapture.Reset();
+				EnableVisualCapture(VisualCaptureEnabled::WidgetPickingMode);
 			}
-			// If we're using the drawing picking mode enable it!
-			else if (PickingMode == EWidgetPickingMode::Drawable)
+			else if (PickingMode != EWidgetPickingMode::None && !VisualCapture.IsEnabled())
 			{
-				VisualCapture.Enable();
+				VisualCapture.Reset();
 			}
 		}
 	}
 
+	/** Enable the VisualCapture. */
+	void EnableVisualCapture(FName Enabler);
+
+	/** Disable the VisualCapture if there are other enablers. */
+	void DisableVisualCapture(FName Enabler);
+
 	/** Callback to see whether the "Snapshot Target" combo should be enabled */
 	bool IsSnapshotTargetComboEnabled() const;
 
@@ -491,8 +513,11 @@ private:
 #if WITH_SLATE_DEBUGGING
 	TWeakPtr<SWidgetHittestGrid> WidgetHittestGrid;
 #endif
+	TSharedPtr<SVisualElementsTab> VisualElementsTabWidget;
 
 	FVisualTreeCapture VisualCapture;
+	/** Controls whether VisualCapture is active. Each entry is a named reason for keeping it enabled. */
+	TArray<FName> VisualCaptureEnablers;
 
 private:
 	float SnapshotDelay;
@@ -583,6 +608,7 @@ void SWidgetReflector::Construct( const FArguments& InArgs )
 					->AddTab(WidgetReflectorTabID::HittestGrid, ETabState::ClosedTab)
 #endif
 					->AddTab(WidgetReflectorTabID::WidgetList, ETabState::ClosedTab)
+					->AddTab(WidgetReflectorTabID::VisualElements, ETabState::ClosedTab)
 				)
 			)
 #if WITH_EDITOR
@@ -638,6 +664,9 @@ void SWidgetReflector::Construct( const FArguments& InArgs )
 	RegisterTrackedTabSpawner(WidgetReflectorTabID::WidgetList, FOnSpawnTab::CreateSP(this, &SWidgetReflector::SpawnWidgetList))
 		.SetDisplayName(LOCTEXT("WidgetListTab", "Widget List"));
 
+	RegisterTrackedTabSpawner(WidgetReflectorTabID::VisualElements, FOnSpawnTab::CreateSP(this, &SWidgetReflector::SpawnVisualElementsTab))
+		.SetDisplayName(LOCTEXT("VisualElementsTab", "Elements"));
+
 #if WITH_EDITOR
 	if (GIsEditor)
 	{
@@ -1109,6 +1138,18 @@ TSharedRef<SDockTab> SWidgetReflector::SpawnWidgetList(const FSpawnTabArgs& Args
 		];
 }
 
+TSharedRef<SDockTab> SWidgetReflector::SpawnVisualElementsTab(const FSpawnTabArgs& Args)
+{
+	return SNew(SDockTab)
+		.Label(LOCTEXT("VisualElementsTab", "Draw Elements"))
+		[
+			SAssignNew(VisualElementsTabWidget, SVisualElementsTab, &VisualCapture)
+			.OnWidgetSelected(this, &SWidgetReflector::SelectLiveWidget)
+			.FilterWidgets(this, &SWidgetReflector::GetSelectedLiveWidgets)
+			.IsPickingActive(this, &SWidgetReflector::IsDrawablePickingActive)
+		];
+}
+
 void SWidgetReflector::HandleTabManagerPersistLayout(const TSharedRef<FTabManager::FLayout>& LayoutToSave)
 {
 #if WITH_EDITOR
@@ -1215,6 +1256,12 @@ void SWidgetReflector::Tick( const FGeometry& AllottedGeometry, const double InC
 
 		TakeSnapshot();
 	}
+
+	//if (VisualCapture.IsEnabled() && VisualElementsTabWidget.IsValid())
+	//{
+	//	VisualElementsTabWidget->SetDrawablePickingActive(PickingMode == EWidgetPickingMode::Drawable);
+	//	VisualElementsTabWidget->Refresh();
+	//}
 }
 
 
@@ -1259,8 +1306,8 @@ int32 SWidgetReflector::Visualize( const FWidgetPath& InWidgetsToVisualize, FSla
 
 	if (PickingMode == EWidgetPickingMode::Drawable)
 	{
-		TSharedPtr<FVisualTreeSnapshot> Tree = VisualCapture.GetVisualTreeForWindow(OutDrawElements.GetPaintWindow());
-		if (Tree.IsValid())
+		TSharedPtr<const FVisualTreeSnapshot> Tree = VisualCapture.GetVisualTreeForWindow(OutDrawElements.GetPaintWindow());
+		if (Tree)
 		{
 			const FVector2f AbsPoint = FSlateApplication::Get().GetCursorPos();
 			const FVector2f WindowPoint = AbsPoint - OutDrawElements.GetPaintWindow()->GetPositionInScreen();
@@ -1317,6 +1364,24 @@ void SWidgetReflector::SelectLiveWidget(TSharedPtr<const SWidget> InWidget)
 	}
 }
 
+TArray<TSharedPtr<const SWidget>> SWidgetReflector::GetSelectedLiveWidgets() const
+{
+	TArray<TSharedPtr<const SWidget>> Result;
+	if (ReflectorTree)
+	{
+		TArray<TSharedRef<FWidgetReflectorNodeBase>> SelectedItems;
+		ReflectorTree->GetSelectedItems(SelectedItems);
+		for (const TSharedRef<FWidgetReflectorNodeBase>& Item : SelectedItems)
+		{
+			if (TSharedPtr<const SWidget> Widget = Item->GetLiveWidget())
+			{
+				Result.Add(Widget);
+			}
+		}
+	}
+	return Result;
+}
+
 namespace WidgetReflectorRecursive
 {
 	bool FindNodeWithReflectionData(const TArray<TSharedRef<FWidgetReflectorNodeBase>>& NodeBase, TArray<TSharedRef<FWidgetReflectorNodeBase>>& Result)
@@ -1665,6 +1730,24 @@ void SWidgetReflector::HandlePickButtonClicked(EWidgetPickingMode InPickingMode)
 	}
 }
 
+void SWidgetReflector::EnableVisualCapture(FName Enabler)
+{
+	VisualCaptureEnablers.AddUnique(Enabler);
+	if (!VisualCapture.IsEnabled())
+	{
+		VisualCapture.Enable();
+	}
+}
+
+void SWidgetReflector::DisableVisualCapture(FName Enabler)
+{
+	VisualCaptureEnablers.RemoveSingleSwap(Enabler);
+	if (VisualCaptureEnablers.Num() == 0 && VisualCapture.IsEnabled())
+	{
+		VisualCapture.Disable();
+	}
+}
+
 bool SWidgetReflector::IsSnapshotTargetComboEnabled() const
 {
 	if (bIsPendingDelayedSnapshot)
@@ -1780,6 +1863,7 @@ void SWidgetReflector::TakeSnapshot()
 		SnapshotData.TakeSnapshot(bRequestNavigationSimulation);
 
 		// Rebuild the reflector tree from the snapshot data
+		VisualCapture.Reset();
 		SelectedNodes.Reset();
 		PickedWidgetPath.Reset();
 		ReflectorTreeRoot = FilteredTreeRoot = SnapshotData.GetWindowsRef();
