Widget Reflector: New scroll bars

#Synced-CL 51887789

diff --git a/Engine/Source/Developer/SlateReflector/Private/Widgets/SWidgetReflector.cpp b/Engine/Source/Developer/SlateReflector/Private/Widgets/SWidgetReflector.cpp
--- a/Engine/Source/Developer/SlateReflector/Private/Widgets/SWidgetReflector.cpp
+++ b/Engine/Source/Developer/SlateReflector/Private/Widgets/SWidgetReflector.cpp
@@ -21,6 +21,7 @@
 #include "Widgets/Images/SImage.h"
 #include "Widgets/Text/STextBlock.h"
 #include "Widgets/Layout/SBox.h"
+#include "Widgets/Layout/SGridPanel.h"
 #include "Widgets/Layout/SSeparator.h"
 #include "Widgets/Input/SButton.h"
 #include "Widgets/SToolTip.h"
@@ -192,6 +193,8 @@ class SWidgetReflector : public ::SWidgetReflector
 public:
 	virtual ~SWidgetReflector();
 
+	virtual void UpdateMaxColumnWidth(FName InColumn, float InWidth) override;
+
 private:
 	virtual void Construct( const FArguments& InArgs) override;
 	//~ End ::SWidgetReflector implementation
@@ -216,6 +219,8 @@ private:
 	TSharedRef<SDockTab> SpawnWidgetList(const FSpawnTabArgs& Args);
 	TSharedRef<SDockTab> SpawnVisualElementsTab(const FSpawnTabArgs& Args);
 
+	float GetColumnWidth(FName InColumn, float InDefault) const;
+
 public:
 
 	void HandleTabManagerPersistLayout(const TSharedRef<FTabManager::FLayout>& LayoutToSave);
@@ -527,6 +532,8 @@ private:
 	bool bTransactedPIEObject;
 #endif
 	double TimeOfScheduledSnapshot;
+
+	TMap<FName, float> ColumnWidths;
 };
 
 void SWidgetReflector::Construct( const FArguments& InArgs )
@@ -775,6 +782,32 @@ TSharedRef<SDockTab> SWidgetReflector::SpawnSlateOptionWidgetTab(const FSpawnTab
 	return SpawnedTab;
 }
 
+float SWidgetReflector::GetColumnWidth(FName InColumn, float InDefault) const
+{
+	if (const float* Width = ColumnWidths.Find(InColumn))
+	{
+		return FMath::Max(*Width, InDefault);
+	}
+
+	return InDefault;
+}
+
+void SWidgetReflector::UpdateMaxColumnWidth(FName InColumn, float InWidth)
+{
+	float& Width = ColumnWidths.FindOrAdd(InColumn);
+
+	if (InWidth > Width)
+	{
+		Width = InWidth;
+
+		if (ReflectorTree.IsValid())
+		{
+			ReflectorTree->GetHeaderRow()->SetColumnWidth(InColumn, InWidth);
+			ReflectorTree->GetHeaderRow()->OnColumnsChanged();
+		}
+	}
+}
+
 TSharedRef<SDockTab> SWidgetReflector::SpawnWidgetHierarchyTab(const FSpawnTabArgs& Args)
 {
 	FSlimHorizontalToolBarBuilder ToolbarBuilderGlobal(TSharedPtr<const FUICommandList>(), FMultiBoxCustomization::None);
@@ -906,6 +939,9 @@ TSharedRef<SDockTab> SWidgetReflector::SpawnWidgetHierarchyTab(const FSpawnTabAr
 	ToolbarBuilderGlobal.EndSection();
 	ToolbarBuilderGlobal.SetStyle(&FAppStyle::Get(), "SlimToolBar");
 
+	TSharedPtr<SScrollBar> VerticalScrollBar;
+	TSharedPtr<SScrollBar> HorizontalScrollBar;
+
 	TSharedRef<SDockTab> SpawnedTab = SNew(SDockTab)
 
 		.Label(LOCTEXT("WidgetHierarchyTab", "Widget Hierarchy"))
@@ -953,72 +989,120 @@ TSharedRef<SDockTab> SWidgetReflector::SpawnWidgetHierarchyTab(const FSpawnTabAr
 				.Padding(0.f)
 				.BorderImage(FCoreStyle::Get().GetBrush("ToolPanel.GroupBorder"))
 				[
-					// The tree view that shows all the info that we capture.
-					SAssignNew(ReflectorTree, SReflectorTree)
-					.TreeItemsSource(&FilteredTreeRoot)
-					.OnGenerateRow(this, &SWidgetReflector::HandleReflectorTreeGenerateRow)
-					.OnGetChildren(this, &SWidgetReflector::HandleReflectorTreeGetChildren)
-					.OnSelectionChanged(this, &SWidgetReflector::HandleReflectorTreeSelectionChanged)
-					.OnContextMenuOpening(this, &SWidgetReflector::HandleReflectorTreeContextMenuPtr)
-					.OnMouseButtonClick(this, &SWidgetReflector::HandleReflectorTreeOnMouseClick)
-					.HighlightParentNodesForSelection(true)
-					.HeaderRow
-					(
-						SNew(SHeaderRow)
-						.CanSelectGeneratedColumn(true)
-						.HiddenColumnsList(HiddenColumnsList)
-						.OnHiddenColumnsListChanged(this, &SWidgetReflector::HandleReflectorTreeHiddenColumnsListChanged)
-
-						+SHeaderRow::Column(SReflectorTreeWidgetItem::NAME_WidgetName)
-						.DefaultLabel(LOCTEXT("WidgetName", "Widget Name"))
-						.FillWidth(1.f)
-						.ShouldGenerateWidget(true)
-
-						+SHeaderRow::Column(SReflectorTreeWidgetItem::NAME_ForegroundColor)
-						.DefaultLabel(LOCTEXT("ForegroundColor", "FG"))
-						.DefaultTooltip(LOCTEXT("ForegroundColorToolTip", "Foreground Color"))
-						.FixedWidth(24.0f)
-
-						+SHeaderRow::Column(SReflectorTreeWidgetItem::NAME_Visibility)
-						.DefaultLabel(LOCTEXT("Visibility", "Visibility"))
-						.DefaultTooltip(LOCTEXT("VisibilityTooltip", "Visibility"))
-						.FillSized(125.0f)
+					SNew(SGridPanel)
+					.FillColumn(0, 1.f)
+					.FillRow(0, 1.f)
+
+					+ SGridPanel::Slot(0, 1)
+					[
+						SAssignNew(HorizontalScrollBar, SScrollBar)
+						.Orientation(Orient_Horizontal)
+						.AlwaysShowScrollbar(false)
+						.AlwaysShowScrollbarTrack(true)
+					]
+
+					+ SGridPanel::Slot(1, 0)
+					[
+						SAssignNew(VerticalScrollBar, SScrollBar)
+						.Orientation(Orient_Vertical)
+						.AlwaysShowScrollbar(false)
+						.AlwaysShowScrollbarTrack(true)
+					]
+					
+					+ SGridPanel::Slot(0, 0)
+					.HAlign(HAlign_Fill)
+					.VAlign(VAlign_Fill)
+					[
+						SNew(SScrollBox)
+						.ExternalScrollbar(HorizontalScrollBar)
+						.Orientation(Orient_Horizontal)
+						.ScrollBarAlwaysVisible(true)
+						
+						+ SScrollBox::Slot()
+						.AutoSize()
+						.HAlign(HAlign_Fill)
+						.VAlign(VAlign_Fill)
+						[
+							SNew(SScrollBox)
+							.ExternalScrollbar(VerticalScrollBar)
+							.Orientation(Orient_Vertical)
+							.ScrollBarAlwaysVisible(true)
+
+							+ SScrollBox::Slot()
+							.AutoSize()
+							.HAlign(HAlign_Fill)
+							.VAlign(VAlign_Fill)
+							[
+								// The tree view that shows all the info that we capture.
+								SAssignNew(ReflectorTree, SReflectorTree)
+								.TreeItemsSource(&FilteredTreeRoot)
+								.OnGenerateRow(this, &SWidgetReflector::HandleReflectorTreeGenerateRow)
+								.OnGetChildren(this, &SWidgetReflector::HandleReflectorTreeGetChildren)
+								.OnSelectionChanged(this, &SWidgetReflector::HandleReflectorTreeSelectionChanged)
+								.OnContextMenuOpening(this, &SWidgetReflector::HandleReflectorTreeContextMenuPtr)
+								.OnMouseButtonClick(this, &SWidgetReflector::HandleReflectorTreeOnMouseClick)
+								.HighlightParentNodesForSelection(true)
+								.ScrollBarStyle(&FWidgetReflectorStyle::Get().GetWidgetStyle<FScrollBarStyle>("ScrollBar.Hidden"))
+								.HeaderRow
+								(
+									SNew(SHeaderRow)
+									.CanSelectGeneratedColumn(true)
+									.HiddenColumnsList(HiddenColumnsList)
+									.OnHiddenColumnsListChanged(this, &SWidgetReflector::HandleReflectorTreeHiddenColumnsListChanged)
+
+									+SHeaderRow::Column(SReflectorTreeWidgetItem::NAME_WidgetName)
+									.DefaultLabel(LOCTEXT("WidgetName", "Widget Name"))								
+									.ShouldGenerateWidget(true)
+									.ManualWidth(this, &SWidgetReflector::GetColumnWidth, SReflectorTreeWidgetItem::NAME_WidgetName, 100.f)
+
+									+SHeaderRow::Column(SReflectorTreeWidgetItem::NAME_ForegroundColor)
+									.DefaultLabel(LOCTEXT("ForegroundColor", "FG"))
+									.DefaultTooltip(LOCTEXT("ForegroundColorToolTip", "Foreground Color"))
+									.FixedWidth(24.0f)
+
+									+SHeaderRow::Column(SReflectorTreeWidgetItem::NAME_Visibility)
+									.DefaultLabel(LOCTEXT("Visibility", "Visibility"))
+									.DefaultTooltip(LOCTEXT("VisibilityTooltip", "Visibility"))
+									.FixedWidth(125.0f)
 						
-						+ SHeaderRow::Column(SReflectorTreeWidgetItem::NAME_Enabled)
-						.DefaultLabel(LOCTEXT("Enabled", "Enabled"))
-						.DefaultTooltip(LOCTEXT("EnabledToolTip", "Enabled"))
-						.FillSized(60.0f)
-
-						+ SHeaderRow::Column(SReflectorTreeWidgetItem::NAME_Focusable)
-						.DefaultLabel(LOCTEXT("Focus", "Focus"))
-						.DefaultTooltip(LOCTEXT("FocusableTooltip", "Focusability (Note that for hit-test directional navigation to work it must be Focusable and \"Visible\"!)"))
-						.FillSized(60.0f)
-
-						+ SHeaderRow::Column(SReflectorTreeWidgetItem::NAME_HasActiveTimer)
-						.DefaultLabel(LOCTEXT("HasActiveTimer", "Timer"))
-						.DefaultTooltip(LOCTEXT("HasActiveTimerTooltip", "Has Active Timer"))
-						.FillSized(60.0f)
-
-						+SHeaderRow::Column(SReflectorTreeWidgetItem::NAME_Clipping)
-						.DefaultLabel(LOCTEXT("Clipping", "Clipping" ))
-						.FillSized(100.0f)
-
-						+ SHeaderRow::Column(SReflectorTreeWidgetItem::NAME_LayerId)
-						.DefaultLabel(LOCTEXT("LayerId", "LayerId"))
-						.FillSized(35.f)
-
-						+ SHeaderRow::Column(SReflectorTreeWidgetItem::NAME_ActualSize)
-						.DefaultLabel(LOCTEXT("ActualSize", "Size"))
-						.FillSized(100.0f)
-
-						+SHeaderRow::Column(SReflectorTreeWidgetItem::NAME_WidgetInfo)
-						.DefaultLabel(LOCTEXT("Source", "Source" ))
-						.FillSized(200.f)
-
-						+SHeaderRow::Column(SReflectorTreeWidgetItem::NAME_Address)
-						.DefaultLabel( LOCTEXT("Address", "Address") )
-						.FillSized(170.0f)
-					)
+									+ SHeaderRow::Column(SReflectorTreeWidgetItem::NAME_Enabled)
+									.DefaultLabel(LOCTEXT("Enabled", "Enabled"))
+									.DefaultTooltip(LOCTEXT("EnabledToolTip", "Enabled"))
+									.FixedWidth(28.0f)
+
+									+ SHeaderRow::Column(SReflectorTreeWidgetItem::NAME_Focusable)
+									.DefaultLabel(LOCTEXT("Focus", "Focus"))
+									.DefaultTooltip(LOCTEXT("FocusableTooltip", "Focusability (Note that for hit-test directional navigation to work it must be Focusable and \"Visible\"!)"))
+									.FixedWidth(28.0f)
+
+									+ SHeaderRow::Column(SReflectorTreeWidgetItem::NAME_HasActiveTimer)
+									.DefaultLabel(LOCTEXT("HasActiveTimer", "Timer"))
+									.DefaultTooltip(LOCTEXT("HasActiveTimerTooltip", "Has Active Timer"))
+									.FixedWidth(28.0f)
+
+									+SHeaderRow::Column(SReflectorTreeWidgetItem::NAME_Clipping)
+									.DefaultLabel(LOCTEXT("Clipping", "Clipping" ))
+									.FixedWidth(36.0f)
+
+									+ SHeaderRow::Column(SReflectorTreeWidgetItem::NAME_LayerId)
+									.DefaultLabel(LOCTEXT("LayerId", "LayerId"))
+									.ManualWidth(this, &SWidgetReflector::GetColumnWidth, SReflectorTreeWidgetItem::NAME_LayerId, 100.f)
+
+									+ SHeaderRow::Column(SReflectorTreeWidgetItem::NAME_ActualSize)
+									.DefaultLabel(LOCTEXT("ActualSize", "Size"))
+									.ManualWidth(this, &SWidgetReflector::GetColumnWidth, SReflectorTreeWidgetItem::NAME_ActualSize, 100.f)
+
+									+SHeaderRow::Column(SReflectorTreeWidgetItem::NAME_WidgetInfo)
+									.DefaultLabel(LOCTEXT("Source", "Source" ))
+									.ManualWidth(this, &SWidgetReflector::GetColumnWidth, SReflectorTreeWidgetItem::NAME_WidgetInfo, 200.f)
+
+									+SHeaderRow::Column(SReflectorTreeWidgetItem::NAME_Address)
+									.DefaultLabel( LOCTEXT("Address", "Address") )
+									.FixedWidth(150.f)
+								)
+							]
+						]
+					]
 				]
 			]
 		];
@@ -1219,6 +1303,7 @@ void SWidgetReflector::SetUIMode(const EWidgetReflectorUIMode InNewMode)
 		PickedWidgetPath.Reset();
 		ReflectorTreeRoot.Reset();
 		FilteredTreeRoot.Reset();
+		ColumnWidths.Empty();
 		ReflectorTree->RequestTreeRefresh();
 
 		if (CurrentUIMode == EWidgetReflectorUIMode::Snapshot)
@@ -1279,6 +1364,7 @@ void SWidgetReflector::SetWidgetsToVisualize( const FWidgetPath& InWidgetsToVisu
 	FilteredTreeRoot.Reset();
 	PickedWidgetPath.Reset();
 	SelectedNodes.Reset();
+	ColumnWidths.Empty();
 
 	if (InWidgetsToVisualize.IsValid())
 	{
@@ -1427,6 +1513,7 @@ void SWidgetReflector::SetNodesAsReflectorTreeRoot(TArray<TSharedRef<FWidgetRefl
 	{
 		FilteredTreeRoot.Reset();
 		FilteredTreeRoot.Append(FilteredNodes);
+		ColumnWidths.Empty();
 		ReflectorTree->RequestTreeRefresh();
 		TSharedPtr<FWidgetReflectorNodeBase> FirstNodeParent = FilteredNodes[0]->GetParentNode();
 		if (FirstNodeParent.IsValid())
@@ -1867,6 +1954,7 @@ void SWidgetReflector::TakeSnapshot()
 		SelectedNodes.Reset();
 		PickedWidgetPath.Reset();
 		ReflectorTreeRoot = FilteredTreeRoot = SnapshotData.GetWindowsRef();
+		ColumnWidths.Empty();
 		ReflectorTree->RequestTreeRefresh();
 
 		WidgetSnapshotVisualizer->SnapshotDataUpdated();
@@ -1962,6 +2050,7 @@ void SWidgetReflector::HandleRemoteSnapshotReceived(const TArray<uint8>& InSnaps
 	SelectedNodes.Reset();
 	PickedWidgetPath.Reset();
 	ReflectorTreeRoot = FilteredTreeRoot = SnapshotData.GetWindowsRef();
+	ColumnWidths.Empty();
 	ReflectorTree->RequestTreeRefresh();
 
 	WidgetSnapshotVisualizer->SnapshotDataUpdated();
@@ -1994,6 +2083,7 @@ void SWidgetReflector::HandleLoadSnapshotButtonClicked()
 
 			// Rebuild the reflector tree from the snapshot data
 			ReflectorTreeRoot = SnapshotData.GetWindowsRef();
+			ColumnWidths.Empty();
 			ReflectorTree->RequestTreeRefresh();
 
 			WidgetSnapshotVisualizer->SnapshotDataUpdated();
@@ -2116,6 +2206,7 @@ void SWidgetReflector::HandleAvailableSnapshotComboSelectionChanged(TSharedPtr<F
 TSharedRef<ITableRow> SWidgetReflector::HandleReflectorTreeGenerateRow( TSharedRef<FWidgetReflectorNodeBase> InReflectorNode, const TSharedRef<STableViewBase>& OwnerTable )
 {
 	return SNew(SReflectorTreeWidgetItem, OwnerTable)
+		.WidgetReflector(SharedThis(this))
 		.WidgetInfoToVisualize(InReflectorNode)
 		.ToolTip(GenerateToolTipForReflectorNode(InReflectorNode))
 		.SourceCodeAccessor(SourceAccessDelegate)
@@ -2276,6 +2367,7 @@ void SWidgetReflector::HandleBreadcrumbOnClick(const TSharedRef<FWidgetReflector
 	FilteredTreeRoot.Reset();
 	FilteredTreeRoot.Add(InReflectorNode);
 	BreadCrumb->PopCrumb();
+	ColumnWidths.Empty();
 	ReflectorTree->RequestTreeRefresh();
 }
 
@@ -2343,6 +2435,7 @@ void SWidgetReflector::HandleResetFilteredTreeRoot()
 	BreadCrumb->ClearCrumbs();
 	bFilterReflectorTreeRootWithUMG = false;
 	UpdateFilteredTreeRoot();
+	ColumnWidths.Empty();
 	ReflectorTree->RequestTreeRefresh();
 }
 
@@ -2350,6 +2443,7 @@ void SWidgetReflector::HandleStartTreeWithUMG()
 {
 	bFilterReflectorTreeRootWithUMG = !bFilterReflectorTreeRootWithUMG;
 	UpdateFilteredTreeRoot();
+	ColumnWidths.Empty();
 	ReflectorTree->RequestTreeRefresh();
 }
 
diff --git a/Engine/Source/Developer/SlateReflector/Private/Widgets/SWidgetReflector.h b/Engine/Source/Developer/SlateReflector/Private/Widgets/SWidgetReflector.h
--- a/Engine/Source/Developer/SlateReflector/Private/Widgets/SWidgetReflector.h
+++ b/Engine/Source/Developer/SlateReflector/Private/Widgets/SWidgetReflector.h
@@ -30,5 +30,7 @@ public:
 
 	virtual void Construct( const FArguments& InArgs ) = 0;
 	virtual ~SWidgetReflector() = default;
+
+	virtual void UpdateMaxColumnWidth(FName InColumn, float InWidth) = 0;
 };
 
diff --git a/Engine/Source/Developer/SlateReflector/Private/Widgets/SWidgetReflectorTreeWidgetItem.cpp b/Engine/Source/Developer/SlateReflector/Private/Widgets/SWidgetReflectorTreeWidgetItem.cpp
--- a/Engine/Source/Developer/SlateReflector/Private/Widgets/SWidgetReflectorTreeWidgetItem.cpp
+++ b/Engine/Source/Developer/SlateReflector/Private/Widgets/SWidgetReflectorTreeWidgetItem.cpp
@@ -7,6 +7,7 @@
 #include "HAL/PlatformApplicationMisc.h"
 #include "Widgets/SInvalidationPanel.h"
 #include "Widgets/Input/SCheckBox.h"
+#include "Widgets/SWidgetReflector.h"
 
 #if WITH_EDITOR
 #include "Blueprint/UserWidget.h"
@@ -32,6 +33,7 @@ FName SReflectorTreeWidgetItem::NAME_Address(TEXT("Address"));
 
 void SReflectorTreeWidgetItem::Construct(const FArguments& InArgs, const TSharedRef<STableViewBase>& InOwnerTableView)
 {
+	this->WidgetReflectorWeak = InArgs._WidgetReflector;
 	this->WidgetInfo = InArgs._WidgetInfoToVisualize;
 	this->OnAccessSourceCode = InArgs._SourceCodeAccessor;
 	this->OnAccessAsset = InArgs._AssetAccessor;
@@ -44,7 +46,36 @@ void SReflectorTreeWidgetItem::Construct(const FArguments& InArgs, const TShared
 }
 
 BEGIN_SLATE_FUNCTION_BUILD_OPTIMIZATION
+
+// Instead of changing the big method below, wrap it with this one to update the maximum width needed for the column.
 TSharedRef<SWidget> SReflectorTreeWidgetItem::GenerateWidgetForColumn(const FName& ColumnName)
+{
+	TSharedRef<SWidget> Widget = GenerateWidgetForColumn_Internal(ColumnName);
+	
+	Widget->SlatePrepass();
+
+	if (TSharedPtr<SWidgetReflector> WidgetReflector = WidgetReflectorWeak.Pin())
+	{
+		const float Width = Widget->GetDesiredSize().X;
+
+		// We cannot do this until later because the tree depth is not yet known, so defer it to OnPaint.
+		if (ColumnName == SReflectorTreeWidgetItem::NAME_WidgetName)
+		{
+			NameColumnWidth = Width;
+			bNeedsNameColumnWidthUpdate = true;
+		}
+		else if (ColumnName == SReflectorTreeWidgetItem::NAME_LayerId
+			|| ColumnName == SReflectorTreeWidgetItem::NAME_ActualSize
+			|| ColumnName == SReflectorTreeWidgetItem::NAME_WidgetInfo)
+		{
+			WidgetReflector->UpdateMaxColumnWidth(ColumnName, Width);
+		}
+	}
+
+	return Widget;
+}
+
+TSharedRef<SWidget> SReflectorTreeWidgetItem::GenerateWidgetForColumn_Internal(const FName& ColumnName)
 {
 	SReflectorTreeWidgetItem* Self = this;
 	auto BuildCheckBox = [Self](bool bIsChecked)
@@ -52,7 +83,7 @@ TSharedRef<SWidget> SReflectorTreeWidgetItem::GenerateWidgetForColumn(const FNam
 			return SNew(SBox)
 				.HAlign(HAlign_Center)
 				.VAlign(VAlign_Center)
-				.Padding(FMargin(2.0f, 0.0f))
+				.Padding(FMargin(4.0f, 0.0f))
 				[
 					SNew(SCheckBox)
 					.Style(FCoreStyle::Get(), TEXT("WidgetReflector.FocusableCheck"))
@@ -156,7 +187,7 @@ TSharedRef<SWidget> SReflectorTreeWidgetItem::GenerateWidgetForColumn(const FNam
 			{
 				HorizontalBox->AddSlot()
 					.AutoWidth()
-					.Padding(2.0f, 0.0f)
+					.Padding(4.0f, 0.0f)
 					.VAlign(VAlign_Center)
 					[
 						SNew(SColorBlock)
@@ -175,7 +206,7 @@ TSharedRef<SWidget> SReflectorTreeWidgetItem::GenerateWidgetForColumn(const FNam
 		{
 			HorizontalBox->AddSlot()
 			.AutoWidth()
-			.Padding(2.0f, 0.0f)
+			.Padding(4.0f, 0.0f)
 			.VAlign(VAlign_Center)
 			[
 				SNew(STextBlock)
@@ -185,7 +216,7 @@ TSharedRef<SWidget> SReflectorTreeWidgetItem::GenerateWidgetForColumn(const FNam
 
 		HorizontalBox->AddSlot()
 		.AutoWidth()
-		.Padding(2.0f, 0.0f)
+		.Padding(4.0f, 0.0f, 9.f, 0.f)
 		.VAlign(VAlign_Center)
 		[
 			SNew(STextBlock)
@@ -208,7 +239,7 @@ TSharedRef<SWidget> SReflectorTreeWidgetItem::GenerateWidgetForColumn(const FNam
 				.AutoWidth()
 				.HAlign(HAlign_Left)
 				.VAlign(VAlign_Center)
-				.Padding(FMargin(2.0f, 0.0f))
+				.Padding(FMargin(4.f, 0.f, 2.f, 0.f))
 				[
 					SNew(SHyperlink)
 					.Text(SourceClass->GetDisplayNameText())
@@ -223,20 +254,34 @@ TSharedRef<SWidget> SReflectorTreeWidgetItem::GenerateWidgetForColumn(const FNam
 					SNew(STextBlock)
 					.Text(LOCTEXT("Source_In", "in"))
 				];
+			
+			Info->AddSlot()
+				.AutoWidth()
+				.HAlign(HAlign_Left)
+				.VAlign(VAlign_Center)
+				.Padding(FMargin(2.f, 0.f, 4.f, 0.f))
+				[
+					SNew(SHyperlink)
+					.Text(WidgetInfo->GetWidgetReadableLocation())
+					.OnNavigate(this, &SReflectorTreeWidgetItem::HandleHyperlinkNavigate, /* Attempt class debug */ false)
+				];
+		}
+		else
+		{
+#endif
+			Info->AddSlot()
+				.AutoWidth()
+				.HAlign(HAlign_Left)
+				.VAlign(VAlign_Center)
+				.Padding(FMargin(2.f, 0.f))
+				[
+					SNew(SHyperlink)
+					.Text(WidgetInfo->GetWidgetReadableLocation())
+					.OnNavigate(this, &SReflectorTreeWidgetItem::HandleHyperlinkNavigate, /* Attempt class debug */ false)
+				];
+#if WITH_EDITOR
 		}
 #endif
-			
-		Info->AddSlot()
-			.AutoWidth()
-			.HAlign(HAlign_Left)
-			.VAlign(VAlign_Center)
-			.Padding(FMargin(2.0f, 0.0f))
-			[
-				SNew(SHyperlink)
-				.Text(WidgetInfo->GetWidgetReadableLocation())
-				.OnNavigate(this, &SReflectorTreeWidgetItem::HandleHyperlinkNavigate, /* Attempt class debug */ false)
-			];
-
 		return Info;
 	}
 	else if (ColumnName == NAME_Visibility )
@@ -244,7 +289,7 @@ TSharedRef<SWidget> SReflectorTreeWidgetItem::GenerateWidgetForColumn(const FNam
 		return SNew(SBox)
 			.HAlign(HAlign_Center)
 			.VAlign(VAlign_Center)
-			.Padding(FMargin(2.0f, 0.0f))
+			.Padding(FMargin(4.0f, 0.0f))
 			[
 				SNew(STextBlock)
 				.Text(WidgetInfo->GetWidgetVisibilityText())
@@ -272,7 +317,7 @@ TSharedRef<SWidget> SReflectorTreeWidgetItem::GenerateWidgetForColumn(const FNam
 		return SNew(SBox)
 			.HAlign(HAlign_Center)
 			.VAlign(VAlign_Center)
-			.Padding(FMargin(2.0f, 0.0f))
+			.Padding(FMargin(4.0f, 0.0f))
 			[
 				SNew(STextBlock)
 				.Text(WidgetInfo->GetWidgetClippingText())
@@ -283,7 +328,7 @@ TSharedRef<SWidget> SReflectorTreeWidgetItem::GenerateWidgetForColumn(const FNam
 		return SNew(SBox)
 			.HAlign(HAlign_Right)
 			.VAlign(VAlign_Center)
-			.Padding(FMargin(2.0f, 0.0f))
+			.Padding(FMargin(4.0f, 0.0f))
 			[
 				SNew(STextBlock)
 				.Text(FText::Format(LOCTEXT("WidgetLayerIds", "[{0}, {1}]"), FText::AsNumber(WidgetInfo->GetWidgetLayerId()), FText::AsNumber(WidgetInfo->GetWidgetLayerIdOut())))
@@ -299,18 +344,23 @@ TSharedRef<SWidget> SReflectorTreeWidgetItem::GenerateWidgetForColumn(const FNam
 			// Show a checkerboard background so we can see alpha values well
 			.BorderImage(FCoreStyle::Get().GetBrush("Checkerboard"))
 			.VAlign(VAlign_Center)
-			.Padding(FMargin(2.0f, 0.0f))
+			.Padding(FMargin(4.0f, 0.0f))
 			[
 				// Show a color block
 				SNew(SColorBlock)
-					.Color(Foreground.GetSpecifiedColor())
-					.Size(FVector2D(16.0f, 16.0f))
+				.Color(Foreground.GetSpecifiedColor())
+				.Size(FVector2D(16.0f, 16.0f))
 			];
 	}
 	else if (ColumnName == NAME_ActualSize)
 	{
-		return SNew(STextBlock)
-			.Text(FText::FromString(WidgetInfo->GetLocalSize().ToString()));
+		return SNew(SBox)
+			.Padding(FMargin(4.f, 0.f))
+			.VAlign(VAlign_Center)
+			[
+				SNew(STextBlock)
+				.Text(FText::FromString(WidgetInfo->GetLocalSize().ToString()))
+			];
 	}
 	else if (ColumnName == NAME_Address )
 	{
@@ -322,7 +372,7 @@ TSharedRef<SWidget> SReflectorTreeWidgetItem::GenerateWidgetForColumn(const FNam
 
 			+ SHorizontalBox::Slot()
 			.AutoWidth()
-			.Padding(FMargin(2.0f, 0.0f))
+			.Padding(FMargin(4.0f, 0.0f))
 			[
 				SNew(SHyperlink)
 				.ToolTipText(LOCTEXT("ClickToCopyBreakpoint", "Click to copy conditional breakpoint for this instance."))
@@ -332,7 +382,7 @@ TSharedRef<SWidget> SReflectorTreeWidgetItem::GenerateWidgetForColumn(const FNam
 
 			+ SHorizontalBox::Slot()
 			.AutoWidth()
-			.Padding(FMargin(0.f, 0.f, 2.f, 0.f))
+			.Padding(FMargin(0.f, 0.f, 4.f, 0.f))
 			[
 				SNew(SHyperlink)
 				.ToolTipText(LOCTEXT("ClickToCopy", "Click to copy address."))
@@ -347,6 +397,24 @@ TSharedRef<SWidget> SReflectorTreeWidgetItem::GenerateWidgetForColumn(const FNam
 }
 END_SLATE_FUNCTION_BUILD_OPTIMIZATION
 
+int32 SReflectorTreeWidgetItem::OnPaint(const FPaintArgs& Args, const FGeometry& AllottedGeometry, const FSlateRect& MyCullingRect, FSlateWindowElementList& OutDrawElements, int32 LayerId, const FWidgetStyle& InWidgetStyle, bool bParentEnabled) const
+{
+	int32 NewLayerId = SMultiColumnTableRow<TSharedRef<FWidgetReflectorNodeBase>>::OnPaint(Args, AllottedGeometry, MyCullingRect, OutDrawElements, LayerId, InWidgetStyle, bParentEnabled);
+
+	// Deferred update now we have a depth.
+	if (bNeedsNameColumnWidthUpdate)
+	{
+		if (TSharedPtr<SWidgetReflector> WidgetReflector = WidgetReflectorWeak.Pin())
+		{
+			WidgetReflector->UpdateMaxColumnWidth(SReflectorTreeWidgetItem::NAME_WidgetName, NameColumnWidth + (static_cast<float>(GetIndentLevel()) * 16.f));
+		}
+
+		bNeedsNameColumnWidthUpdate = false;
+	}
+
+	return NewLayerId;
+}
+
 void SReflectorTreeWidgetItem::HandleHyperlinkNavigate(bool bAttemptClassDebug)
 {
 #if WITH_EDITOR
diff --git a/Engine/Source/Developer/SlateReflector/Private/Widgets/SWidgetReflectorTreeWidgetItem.h b/Engine/Source/Developer/SlateReflector/Private/Widgets/SWidgetReflectorTreeWidgetItem.h
--- a/Engine/Source/Developer/SlateReflector/Private/Widgets/SWidgetReflectorTreeWidgetItem.h
+++ b/Engine/Source/Developer/SlateReflector/Private/Widgets/SWidgetReflectorTreeWidgetItem.h
@@ -11,6 +11,8 @@
 #include "Widgets/Views/SListView.h"
 #include "Models/WidgetReflectorNode.h"
 
+class SWidgetReflector;
+
 /**
  * Widget that visualizes the contents of a FReflectorNode.
  */
@@ -33,11 +35,12 @@ public:
 	static FName NAME_ActualSize;
 
 	SLATE_BEGIN_ARGS(SReflectorTreeWidgetItem)
-		: _WidgetInfoToVisualize()
+		: _WidgetReflector()
+		, _WidgetInfoToVisualize()
 		, _SourceCodeAccessor()
 		, _AssetAccessor()
 	{ }
-
+		SLATE_ARGUMENT(TSharedPtr<SWidgetReflector>, WidgetReflector)
 		SLATE_ARGUMENT(TSharedPtr<FWidgetReflectorNodeBase>, WidgetInfoToVisualize)
 		SLATE_ARGUMENT(FAccessSourceCode, SourceCodeAccessor)
 		SLATE_ARGUMENT(FAccessAsset, AssetAccessor)
@@ -69,7 +72,12 @@ protected:
 
 	void HandleHyperlinkNavigate(bool bAttemptClassDebug);
 
+	TSharedRef<SWidget> GenerateWidgetForColumn_Internal(const FName& ColumnName);
+
+	virtual int32 OnPaint(const FPaintArgs& Args, const FGeometry& AllottedGeometry, const FSlateRect& MyCullingRect, FSlateWindowElementList& OutDrawElements, int32 LayerId, const FWidgetStyle& InWidgetStyle, bool bParentEnabled) const override;
+
 private:
+	TWeakPtr<SWidgetReflector> WidgetReflectorWeak;
 
 	/** The info about the widget that we are visualizing. */
 	TSharedPtr<FWidgetReflectorNodeBase> WidgetInfo;
@@ -77,4 +85,7 @@ private:
 	FAccessSourceCode OnAccessSourceCode;
 	FAccessAsset OnAccessAsset;
 	FAccessDebugObject OnAccessDebugObject;
+
+	float NameColumnWidth = 0.f;
+	mutable bool bNeedsNameColumnWidthUpdate = false;
 };
