Widget Reflector: Fixed bug relating to expanding the tree

#Synced-CL 51887789

diff --git a/Engine/Source/Developer/SlateReflector/Private/Widgets/SWidgetReflector.cpp b/Engine/Source/Developer/SlateReflector/Private/Widgets/SWidgetReflector.cpp
--- a/Engine/Source/Developer/SlateReflector/Private/Widgets/SWidgetReflector.cpp
+++ b/Engine/Source/Developer/SlateReflector/Private/Widgets/SWidgetReflector.cpp
@@ -193,7 +193,7 @@ class SWidgetReflector : public ::SWidgetReflector
 public:
 	virtual ~SWidgetReflector();
 
-	virtual void UpdateMaxColumnWidth(FName InColumn, float InWidth) override;
+	virtual void UpdateColumnWidth(TSharedRef<FWidgetReflectorNodeBase> InItem, FName InColumn) override;
 
 private:
 	virtual void Construct( const FArguments& InArgs) override;
@@ -792,18 +792,28 @@ float SWidgetReflector::GetColumnWidth(FName InColumn, float InDefault) const
 	return InDefault;
 }
 
-void SWidgetReflector::UpdateMaxColumnWidth(FName InColumn, float InWidth)
+void SWidgetReflector::UpdateColumnWidth(TSharedRef<FWidgetReflectorNodeBase> InItem, FName InColumn)
 {
-	float& Width = ColumnWidths.FindOrAdd(InColumn);
-
-	if (InWidth > Width)
+	if (ReflectorTree.IsValid())
 	{
-		Width = InWidth;
-
-		if (ReflectorTree.IsValid())
+		if (TSharedPtr<SReflectorTreeWidgetItem> Row = StaticCastSharedPtr<SReflectorTreeWidgetItem>(ReflectorTree->WidgetFromItem(InItem)))
 		{
-			ReflectorTree->GetHeaderRow()->SetColumnWidth(InColumn, InWidth);
-			ReflectorTree->GetHeaderRow()->OnColumnsChanged();
+			if (const TSharedRef<SWidget>* WidgetPtr = Row->GetWidgetFromColumnId(InColumn))
+			{
+				const float WidgetWidth = (*WidgetPtr)->GetDesiredSize().X;
+				float& Width = ColumnWidths.FindOrAdd(InColumn);
+
+				if (WidgetWidth > Width)
+				{
+					Width = WidgetWidth;
+
+					if (ReflectorTree.IsValid())
+					{
+						ReflectorTree->GetHeaderRow()->SetColumnWidth(InColumn, WidgetWidth);
+						ReflectorTree->GetHeaderRow()->OnColumnsChanged();
+					}
+				}
+			}
 		}
 	}
 }
diff --git a/Engine/Source/Developer/SlateReflector/Private/Widgets/SWidgetReflector.h b/Engine/Source/Developer/SlateReflector/Private/Widgets/SWidgetReflector.h
--- a/Engine/Source/Developer/SlateReflector/Private/Widgets/SWidgetReflector.h
+++ b/Engine/Source/Developer/SlateReflector/Private/Widgets/SWidgetReflector.h
@@ -8,6 +8,7 @@
 #include "Widgets/SUserWidget.h"
 #include "Framework/Application/IWidgetReflector.h"
 
+class FWidgetReflectorNodeBase;
 class FWidgetSnapshotService;
 
 /**
@@ -31,6 +32,6 @@ public:
 	virtual void Construct( const FArguments& InArgs ) = 0;
 	virtual ~SWidgetReflector() = default;
 
-	virtual void UpdateMaxColumnWidth(FName InColumn, float InWidth) = 0;
+	virtual void UpdateColumnWidth(TSharedRef<FWidgetReflectorNodeBase> InItem, FName InColumn) = 0;
 };
 
diff --git a/Engine/Source/Developer/SlateReflector/Private/Widgets/SWidgetReflectorTreeWidgetItem.cpp b/Engine/Source/Developer/SlateReflector/Private/Widgets/SWidgetReflectorTreeWidgetItem.cpp
--- a/Engine/Source/Developer/SlateReflector/Private/Widgets/SWidgetReflectorTreeWidgetItem.cpp
+++ b/Engine/Source/Developer/SlateReflector/Private/Widgets/SWidgetReflectorTreeWidgetItem.cpp
@@ -52,25 +52,7 @@ TSharedRef<SWidget> SReflectorTreeWidgetItem::GenerateWidgetForColumn(const FNam
 {
 	TSharedRef<SWidget> Widget = GenerateWidgetForColumn_Internal(ColumnName);
 	
-	Widget->SlatePrepass();
-
-	if (TSharedPtr<SWidgetReflector> WidgetReflector = WidgetReflectorWeak.Pin())
-	{
-		const float Width = Widget->GetDesiredSize().X;
-
-		// We cannot do this until later because the tree depth is not yet known, so defer it to OnPaint.
-		if (ColumnName == SReflectorTreeWidgetItem::NAME_WidgetName)
-		{
-			NameColumnWidth = Width;
-			bNeedsNameColumnWidthUpdate = true;
-		}
-		else if (ColumnName == SReflectorTreeWidgetItem::NAME_LayerId
-			|| ColumnName == SReflectorTreeWidgetItem::NAME_ActualSize
-			|| ColumnName == SReflectorTreeWidgetItem::NAME_WidgetInfo)
-		{
-			WidgetReflector->UpdateMaxColumnWidth(ColumnName, Width);
-		}
-	}
+	ColumnsToUpdate.Add(ColumnName);
 
 	return Widget;
 }
@@ -402,14 +384,17 @@ int32 SReflectorTreeWidgetItem::OnPaint(const FPaintArgs& Args, const FGeometry&
 	int32 NewLayerId = SMultiColumnTableRow<TSharedRef<FWidgetReflectorNodeBase>>::OnPaint(Args, AllottedGeometry, MyCullingRect, OutDrawElements, LayerId, InWidgetStyle, bParentEnabled);
 
 	// Deferred update now we have a depth.
-	if (bNeedsNameColumnWidthUpdate)
+	if (!ColumnsToUpdate.IsEmpty() && WidgetInfo.IsValid())
 	{
 		if (TSharedPtr<SWidgetReflector> WidgetReflector = WidgetReflectorWeak.Pin())
 		{
-			WidgetReflector->UpdateMaxColumnWidth(SReflectorTreeWidgetItem::NAME_WidgetName, NameColumnWidth + (static_cast<float>(GetIndentLevel()) * 16.f));
+			for (const FName& ColumnName : ColumnsToUpdate)
+			{
+				WidgetReflector->UpdateColumnWidth(WidgetInfo.ToSharedRef(), ColumnName);
+			}
 		}
 
-		bNeedsNameColumnWidthUpdate = false;
+		ColumnsToUpdate.Empty();
 	}
 
 	return NewLayerId;
diff --git a/Engine/Source/Developer/SlateReflector/Private/Widgets/SWidgetReflectorTreeWidgetItem.h b/Engine/Source/Developer/SlateReflector/Private/Widgets/SWidgetReflectorTreeWidgetItem.h
--- a/Engine/Source/Developer/SlateReflector/Private/Widgets/SWidgetReflectorTreeWidgetItem.h
+++ b/Engine/Source/Developer/SlateReflector/Private/Widgets/SWidgetReflectorTreeWidgetItem.h
@@ -62,6 +62,9 @@ public:
 	// SMultiColumnTableRow overrides
 	virtual TSharedRef<SWidget> GenerateWidgetForColumn( const FName& ColumnName ) override;
 
+	// Expose parent class protected function.
+	using SMultiColumnTableRow<TSharedRef<FWidgetReflectorNodeBase>>::GetWidgetFromColumnId;
+
 protected:
 
 	/** @return The tint of the reflector node */
@@ -86,6 +89,5 @@ private:
 	FAccessAsset OnAccessAsset;
 	FAccessDebugObject OnAccessDebugObject;
 
-	float NameColumnWidth = 0.f;
-	mutable bool bNeedsNameColumnWidthUpdate = false;
+	mutable TSet<FName> ColumnsToUpdate;
 };
