Widget Reflector: Fix for misalgined columns

#Synced-CL 51887789

diff --git a/Engine/Source/Developer/SlateReflector/Private/Widgets/SWidgetReflectorTreeWidgetItem.cpp b/Engine/Source/Developer/SlateReflector/Private/Widgets/SWidgetReflectorTreeWidgetItem.cpp
--- a/Engine/Source/Developer/SlateReflector/Private/Widgets/SWidgetReflectorTreeWidgetItem.cpp
+++ b/Engine/Source/Developer/SlateReflector/Private/Widgets/SWidgetReflectorTreeWidgetItem.cpp
@@ -51,8 +51,14 @@ BEGIN_SLATE_FUNCTION_BUILD_OPTIMIZATION
 TSharedRef<SWidget> SReflectorTreeWidgetItem::GenerateWidgetForColumn(const FName& ColumnName)
 {
 	TSharedRef<SWidget> Widget = GenerateWidgetForColumn_Internal(ColumnName);
-	
-	ColumnsToUpdate.Add(ColumnName);
+
+	if (ColumnName == NAME_WidgetName
+		|| ColumnName == NAME_WidgetInfo
+		|| ColumnName == NAME_LayerId
+		|| ColumnName == NAME_ActualSize)
+	{
+		ColumnsToUpdate.Add(ColumnName);
+	}
 
 	return Widget;
 }
