[WidgetReflector] Add search tool to the widget reflector. Can search all live widgets or only those current filter by the Widget Hiearchy. Search by name, type, or tag.

#Synced-CL 52583850

diff --git a/Engine/Source/Developer/SlateReflector/Private/Widgets/SWidgetReflector.cpp b/Engine/Source/Developer/SlateReflector/Private/Widgets/SWidgetReflector.cpp
--- a/Engine/Source/Developer/SlateReflector/Private/Widgets/SWidgetReflector.cpp
+++ b/Engine/Source/Developer/SlateReflector/Private/Widgets/SWidgetReflector.cpp
@@ -45,6 +45,7 @@
 #include "Widgets/Navigation/SBreadcrumbTrail.h"
 #include "Widgets/Input/SComboBox.h"
 #include "Widgets/SVisualElementsTab.h"
+#include "Widgets/SWidgetSearchPanel.h"
 #include "Widgets/SWidgetSnapshotVisualizer.h"
 #include "WidgetSnapshotService.h"
 #include "Types/ReflectionMetadata.h"
@@ -137,6 +138,7 @@ namespace WidgetReflectorTabID
 	static const FName HittestGrid = "WidgetReflector.HittestGridTab";
 	static const FName WidgetList = "WidgetReflector.WidgetList";
 	static const FName VisualElements = "WidgetReflector.VisualElements";
+	static const FName WidgetSearch = "WidgetReflector.WidgetSearch";
 }
 
 namespace WidgetReflectorText
@@ -218,6 +220,7 @@ private:
 
 	TSharedRef<SDockTab> SpawnWidgetList(const FSpawnTabArgs& Args);
 	TSharedRef<SDockTab> SpawnVisualElementsTab(const FSpawnTabArgs& Args);
+	TSharedRef<SDockTab> SpawnWidgetSearchTab(const FSpawnTabArgs& Args);
 
 	float GetColumnWidth(FName InColumn, float InDefault) const;
 
@@ -264,10 +267,28 @@ public:
 		SourceAccessDelegate = InDelegate;
 	}
 
+	bool HandleSourceAccessDelegate(const FString& FileName, int32 LineNumber, int32 ColumnNumber) const
+	{
+		if (SourceAccessDelegate.IsBound())
+		{
+			return SourceAccessDelegate.Execute(FileName, LineNumber, ColumnNumber);
+		}
+		return false;
+	}
+
 	virtual void SetAssetAccessDelegate(FAccessAsset InDelegate) override
 	{
 		AssetAccessDelegate = InDelegate;
 	}
+
+	bool HandleAssetAccessDelegate(UObject* Asset) const
+	{
+		if (AssetAccessDelegate.IsBound())
+		{
+			return AssetAccessDelegate.Execute(Asset);
+		}
+		return false;
+	}
 	
 	virtual void SetDebugObjectAccessDelegate(FAccessDebugObject InDelegate)override
 	{
@@ -320,6 +341,9 @@ public:
 	/** Clear previous selection and set the selection to the live widget. */
 	void SelectLiveWidget( TSharedPtr<const SWidget> InWidget );
 
+	/** Clear previous selection and select the given node in the reflector tree by expanding its path and scrolling it into view. */
+	void SelectReflectorNode(TSharedRef<FWidgetReflectorNodeBase> InNode);
+
 	/** Get all selected live widgets. */
 	TArray<TSharedPtr<const SWidget>> GetSelectedLiveWidgets() const;
 
@@ -519,6 +543,7 @@ private:
 	TWeakPtr<SWidgetHittestGrid> WidgetHittestGrid;
 #endif
 	TSharedPtr<SVisualElementsTab> VisualElementsTabWidget;
+	TWeakPtr<UE::SlateReflector::SWidgetSearchPanel> SearchPanel;
 
 	FVisualTreeCapture VisualCapture;
 	/** Controls whether VisualCapture is active. Each entry is a named reason for keeping it enabled. */
@@ -610,6 +635,7 @@ void SWidgetReflector::Construct( const FArguments& InArgs )
 					->SetHideTabWell(true)
 					->SetSizeCoefficient(0.3f)
 					->AddTab(WidgetReflectorTabID::SnapshotWidgetPicker, ETabState::ClosedTab)
+					->AddTab(WidgetReflectorTabID::WidgetSearch, ETabState::ClosedTab)
 #if WITH_SLATE_DEBUGGING
 					->AddTab(WidgetReflectorTabID::WidgetEvents, ETabState::ClosedTab)
 					->AddTab(WidgetReflectorTabID::HittestGrid, ETabState::ClosedTab)
@@ -647,12 +673,6 @@ void SWidgetReflector::Construct( const FArguments& InArgs )
 	RegisterTrackedTabSpawner(WidgetReflectorTabID::SlateOptions, FOnSpawnTab::CreateSP(this, &SWidgetReflector::SpawnSlateOptionWidgetTab))
 		.SetDisplayName(LOCTEXT("OptionsTab", "Toolbar"));
 
-	RegisterTrackedTabSpawner(WidgetReflectorTabID::WidgetHierarchy, FOnSpawnTab::CreateSP(this, &SWidgetReflector::SpawnWidgetHierarchyTab))
-		.SetDisplayName(LOCTEXT("WidgetHierarchyTab", "Widget Hierarchy"));
-
-	RegisterTrackedTabSpawner(WidgetReflectorTabID::SnapshotWidgetPicker, FOnSpawnTab::CreateSP(this, &SWidgetReflector::SpawnSnapshotWidgetPicker))
-		.SetDisplayName(LOCTEXT("SnapshotWidgetPickerTab", "Snapshot Widget Picker"));
-
 #if WITH_EDITOR
 	if (GIsEditor)
 	{
@@ -664,16 +684,28 @@ void SWidgetReflector::Construct( const FArguments& InArgs )
 #if WITH_SLATE_DEBUGGING
 	RegisterTrackedTabSpawner(WidgetReflectorTabID::WidgetEvents, FOnSpawnTab::CreateSP(this, &SWidgetReflector::SpawnWidgetEvents))
 		.SetDisplayName(LOCTEXT("WidgetEventsTab", "Widget Events"));
-	RegisterTrackedTabSpawner(WidgetReflectorTabID::HittestGrid, FOnSpawnTab::CreateSP(this, &SWidgetReflector::SpawnWidgetHittestGrid))
-		.SetDisplayName(LOCTEXT("HitTestGridTab", "Hit Test Grid"));
-#endif
+#endif // WITH_SLATE_DEBUGGING
+
+	RegisterTrackedTabSpawner(WidgetReflectorTabID::WidgetHierarchy, FOnSpawnTab::CreateSP(this, &SWidgetReflector::SpawnWidgetHierarchyTab))
+		.SetDisplayName(LOCTEXT("WidgetHierarchyTab", "Widget Hierarchy"));
 
 	RegisterTrackedTabSpawner(WidgetReflectorTabID::WidgetList, FOnSpawnTab::CreateSP(this, &SWidgetReflector::SpawnWidgetList))
 		.SetDisplayName(LOCTEXT("WidgetListTab", "Widget List"));
 
+	RegisterTrackedTabSpawner(WidgetReflectorTabID::WidgetSearch, FOnSpawnTab::CreateSP(this, &SWidgetReflector::SpawnWidgetSearchTab))
+		.SetDisplayName(LOCTEXT("WidgetSearchTab", "Widget Search"));
+
 	RegisterTrackedTabSpawner(WidgetReflectorTabID::VisualElements, FOnSpawnTab::CreateSP(this, &SWidgetReflector::SpawnVisualElementsTab))
 		.SetDisplayName(LOCTEXT("VisualElementsTab", "Elements"));
 
+#if WITH_SLATE_DEBUGGING
+	RegisterTrackedTabSpawner(WidgetReflectorTabID::HittestGrid, FOnSpawnTab::CreateSP(this, &SWidgetReflector::SpawnWidgetHittestGrid))
+		.SetDisplayName(LOCTEXT("HitTestGridTab", "Hit Test Grid"));
+#endif // WITH_SLATE_DEBUGGING
+
+	RegisterTrackedTabSpawner(WidgetReflectorTabID::SnapshotWidgetPicker, FOnSpawnTab::CreateSP(this, &SWidgetReflector::SpawnSnapshotWidgetPicker))
+		.SetDisplayName(LOCTEXT("SnapshotWidgetPickerTab", "Snapshot Widget Picker"));
+
 #if WITH_EDITOR
 	if (GIsEditor)
 	{
@@ -1227,8 +1259,8 @@ TSharedRef<SDockTab> SWidgetReflector::SpawnWidgetList(const FSpawnTabArgs& Args
 		.Label(LOCTEXT("WidgetListTab", "Widget List"))
 		[
 			SNew(SWidgetList)
-			.OnAccessSource(SourceAccessDelegate)
-			.OnAccessAsset(AssetAccessDelegate)
+			.OnAccessSource(FAccessSourceCode::CreateSP(this, &SWidgetReflector::HandleSourceAccessDelegate))
+			.OnAccessAsset(FAccessAsset::CreateSP(this, &SWidgetReflector::HandleAssetAccessDelegate))
 		];
 }
 
@@ -1244,6 +1276,25 @@ TSharedRef<SDockTab> SWidgetReflector::SpawnVisualElementsTab(const FSpawnTabArg
 		];
 }
 
+TSharedRef<SDockTab> SWidgetReflector::SpawnWidgetSearchTab(const FSpawnTabArgs& Args)
+{
+	TSharedRef<SDockTab> SpawnedTab = SNew(SDockTab)
+		.Label(LOCTEXT("WidgetSearchTab", "Widget Search"))
+		[
+			SNew(UE::SlateReflector::SWidgetSearchPanel)
+			.OnSelectNode(this, &SWidgetReflector::SelectReflectorNode)
+			.OnSelectLiveWidget(this, &SWidgetReflector::SelectLiveWidget)
+			.OnGetRootNodes_Lambda([this]() -> TArray<TSharedRef<FWidgetReflectorNodeBase>>
+			{
+				return ReflectorTreeRoot;
+			})
+			.OnAccessSourceCode(FAccessSourceCode::CreateSP(this, &SWidgetReflector::HandleSourceAccessDelegate))
+			.OnAccessAsset(FAccessAsset::CreateSP(this, &SWidgetReflector::HandleAssetAccessDelegate))
+		];
+
+	return SpawnedTab;
+}
+
 void SWidgetReflector::HandleTabManagerPersistLayout(const TSharedRef<FTabManager::FLayout>& LayoutToSave)
 {
 #if WITH_EDITOR
@@ -1458,6 +1509,39 @@ void SWidgetReflector::SelectLiveWidget(TSharedPtr<const SWidget> InWidget)
 	{
 		ReflectorTree->ClearSelection();
 	}
+
+	if (bFound && TabManager.IsValid())
+	{
+		TabManager->TryInvokeTab(WidgetReflectorTabID::WidgetHierarchy);
+	}
+}
+
+void SWidgetReflector::SelectReflectorNode(TSharedRef<FWidgetReflectorNodeBase> InNode)
+{
+	if (!ReflectorTree.IsValid())
+	{
+		return;
+	}
+
+	ReflectorTree->ClearSelection();
+
+	// Walk up the parent chain and expand each ancestor so the node is visible in the tree.
+	TArray<TSharedRef<FWidgetReflectorNodeBase>> PathToExpand;
+	TSharedPtr<FWidgetReflectorNodeBase> Parent = InNode->GetParentNode();
+	while (Parent.IsValid())
+	{
+		PathToExpand.Add(Parent.ToSharedRef());
+		Parent = Parent->GetParentNode();
+	}
+
+	for (int32 Index = PathToExpand.Num() - 1; Index >= 0; --Index)
+	{
+		const TSharedRef<FWidgetReflectorNodeBase>& PathNode = PathToExpand[Index];
+		ReflectorTree->SetItemExpansion(PathNode, true);
+	}
+
+	ReflectorTree->SetSelection(InNode);
+	ReflectorTree->RequestScrollIntoView(InNode);
 }
 
 TArray<TSharedPtr<const SWidget>> SWidgetReflector::GetSelectedLiveWidgets() const
@@ -2219,7 +2303,7 @@ TSharedRef<ITableRow> SWidgetReflector::HandleReflectorTreeGenerateRow( TSharedR
 		.WidgetReflector(SharedThis(this))
 		.WidgetInfoToVisualize(InReflectorNode)
 		.ToolTip(GenerateToolTipForReflectorNode(InReflectorNode))
-		.SourceCodeAccessor(SourceAccessDelegate)
+		.SourceCodeAccessor(FAccessSourceCode::CreateSP(this, &SWidgetReflector::HandleSourceAccessDelegate))
 		.AssetAccessor(AssetAccessDelegate)
 		.DebugObjectAccessor(DebugObjectAccessDelegate);
 }
diff --git a/Engine/Source/Developer/SlateReflector/Private/Widgets/SWidgetSearchPanel.cpp b/Engine/Source/Developer/SlateReflector/Private/Widgets/SWidgetSearchPanel.cpp
new file mode 100644
--- /dev/null
+++ b/Engine/Source/Developer/SlateReflector/Private/Widgets/SWidgetSearchPanel.cpp
@@ -0,0 +1,513 @@
+// Copyright Epic Games, Inc. All Rights Reserved.
+
+#include "Widgets/SWidgetSearchPanel.h"
+
+#include "Framework/Application/SlateApplication.h"
+#include "Models/WidgetReflectorNode.h"
+#include "Styling/CoreStyle.h"
+#include "Styling/StyleColors.h"
+#include "HAL/PlatformApplicationMisc.h"
+#include "Widgets/Input/SButton.h"
+#include "Widgets/Input/SCheckBox.h"
+#include "Widgets/Input/SEditableTextBox.h"
+#include "Widgets/Input/SHyperlink.h"
+#include "Widgets/Layout/SBorder.h"
+#include "Widgets/Layout/SBox.h"
+#include "Widgets/Layout/SSplitter.h"
+#include "Widgets/SBoxPanel.h"
+#include "Widgets/SNullWidget.h"
+#include "Widgets/SWindow.h"
+#include "Widgets/Text/STextBlock.h"
+#include "Widgets/Views/SHeaderRow.h"
+#include "Widgets/Views/STableRow.h"
+
+#define LOCTEXT_NAMESPACE "SWidgetSearchPanel"
+
+namespace UE::SlateReflector
+{
+
+namespace WidgetSearchColumnID
+{
+	static const FName DisplayName = "DisplayName";
+	static const FName Type = "Type";
+	static const FName Tag = "Tag";
+	static const FName Source = "Source";
+	static const FName Address = "Address";
+}
+
+// ─── Result entry ────────────────────────────────────────────────────────────
+
+	/** A single entry displayed in the search results list */
+struct FWidgetSearchResultEntry
+{
+	TSharedRef<FWidgetReflectorNodeBase> Node;
+	FText DisplayName;
+	FText Type;
+	FText Tag;
+	FText Source;
+	FText Address;
+
+	explicit FWidgetSearchResultEntry(TSharedRef<FWidgetReflectorNodeBase> InNode)
+		: Node(InNode)
+	{
+		DisplayName = Node->GetWidgetTypeAndShortName();
+		Type = Node->GetWidgetType();
+		Source = Node->GetWidgetReadableLocation();
+		Address = FText::FromString(FWidgetReflectorNodeUtils::WidgetAddressToString(Node->GetWidgetAddress()));
+
+		// Tag is only available for live widgets; snapshot nodes return nullptr from GetLiveWidget().
+		if (TSharedPtr<SWidget> LiveWidget = Node->GetLiveWidget())
+		{
+			const FName TagName = LiveWidget->GetTag();
+			Tag = TagName.IsNone() ? FText::GetEmpty() : FText::FromName(TagName);
+		}
+	}
+};
+
+// ─── Row widget ──────────────────────────────────────────────────────────────
+
+class SWidgetSearchResultRow : public SMultiColumnTableRow<TSharedPtr<FWidgetSearchResultEntry>>
+{
+public:
+	SLATE_BEGIN_ARGS(SWidgetSearchResultRow) {}
+		SLATE_ARGUMENT(TSharedPtr<FWidgetSearchResultEntry>, Entry)
+		SLATE_ARGUMENT(FAccessSourceCode, SourceCodeAccessor)
+		SLATE_ARGUMENT(FAccessAsset, AssetAccessor)
+	SLATE_END_ARGS()
+
+	void Construct(const FArguments& InArgs, const TSharedRef<STableViewBase>& InOwnerTable)
+	{
+		Entry = InArgs._Entry;
+		OnAccessSourceCode = InArgs._SourceCodeAccessor;
+		OnAccessAsset = InArgs._AssetAccessor;
+		FSuperRowType::Construct(FSuperRowType::FArguments(), InOwnerTable);
+	}
+
+	virtual TSharedRef<SWidget> GenerateWidgetForColumn(const FName& ColumnName) override
+	{
+		if (!Entry.IsValid())
+		{
+			return SNullWidget::NullWidget;
+		}
+
+		if (ColumnName == WidgetSearchColumnID::Source)
+		{
+			return SNew(SBox)
+				.Padding(FMargin(4.f, 2.f))
+				.VAlign(VAlign_Center)
+				[
+					SNew(SHyperlink)
+					.Text(Entry->Source)
+					.ToolTipText(Entry->Source)
+					.OnNavigate(this, &SWidgetSearchResultRow::HandleSourceNavigation)
+				];
+		}
+
+		if (ColumnName == WidgetSearchColumnID::Address)
+		{
+			return SNew(SBox)
+				.Padding(FMargin(4.f, 2.f))
+				.VAlign(VAlign_Center)
+				[
+					SNew(SHorizontalBox)
+
+					+ SHorizontalBox::Slot()
+					.AutoWidth()
+					[
+						SNew(SHyperlink)
+						.ToolTipText(LOCTEXT("ClickToCopyBreakpoint", "Click to copy conditional breakpoint for this instance."))
+						.Text(LOCTEXT("CBP", "[CBP]"))
+						.OnNavigate(this, &SWidgetSearchResultRow::HandleConditionalBreakpoint)
+					]
+
+					+ SHorizontalBox::Slot()
+					.AutoWidth()
+					.Padding(FMargin(4.f, 0.f, 0.f, 0.f))
+					[
+						SNew(SHyperlink)
+						.ToolTipText(LOCTEXT("ClickToCopyAddress", "Click to copy address."))
+						.Text(Entry->Address)
+						.OnNavigate(this, &SWidgetSearchResultRow::HandleClickToCopyAddress)
+					]
+				];
+		}
+
+		FText ColumnText;
+		if (ColumnName == WidgetSearchColumnID::DisplayName)
+		{
+			ColumnText = Entry->DisplayName;
+		}
+		else if (ColumnName == WidgetSearchColumnID::Type)
+		{
+			ColumnText = Entry->Type;
+		}
+		else if (ColumnName == WidgetSearchColumnID::Tag)
+		{
+			ColumnText = Entry->Tag;
+		}
+
+		return SNew(SBox)
+			.Padding(FMargin(4.f, 2.f))
+			.VAlign(VAlign_Center)
+			[
+				SNew(STextBlock)
+				.Text(ColumnText)
+				.ToolTipText(ColumnText)
+			];
+	}
+
+private:
+	void HandleSourceNavigation()
+	{
+		const FAssetData& AssetData = Entry->Node->GetWidgetAssetData();
+		if (AssetData.IsValid() && OnAccessAsset.IsBound())
+		{
+			OnAccessAsset.Execute(AssetData.GetAsset());
+		}
+		else if (OnAccessSourceCode.IsBound())
+		{
+			OnAccessSourceCode.Execute(Entry->Node->GetWidgetFile(), Entry->Node->GetWidgetLineNumber(), 0);
+		}
+	}
+
+	void HandleConditionalBreakpoint()
+	{
+		const FString ConditionalBreakpoint = FString::Printf(TEXT("this == (SWidget*)%s"), *Entry->Address.ToString());
+		FPlatformApplicationMisc::ClipboardCopy(*ConditionalBreakpoint);
+	}
+
+	void HandleClickToCopyAddress()
+	{
+		FPlatformApplicationMisc::ClipboardCopy(*Entry->Address.ToString());
+	}
+
+private:
+	TSharedPtr<FWidgetSearchResultEntry> Entry;
+	FAccessSourceCode OnAccessSourceCode;
+	FAccessAsset OnAccessAsset;
+};
+
+// ─── Panel ───────────────────────────────────────────────────────────────────
+
+void SWidgetSearchPanel::Construct(const FArguments& InArgs)
+{
+	OnSelectNode = InArgs._OnSelectNode;
+	OnSelectLiveWidget = InArgs._OnSelectLiveWidget;
+	OnGetRootNodes = InArgs._OnGetRootNodes;
+	OnAccessSourceCode = InArgs._OnAccessSourceCode;
+	OnAccessAsset = InArgs._OnAccessAsset;
+
+	// Helper: create a fixed-width right-aligned label cell
+	auto MakeLabel = [](FText InText) -> TSharedRef<SWidget>
+	{
+		return SNew(SBox)
+			.WidthOverride(90.f)
+			.HAlign(HAlign_Right)
+			.VAlign(VAlign_Center)
+			.Padding(FMargin(0.f, 0.f, 6.f, 0.f))
+			[
+				SNew(STextBlock).Text(InText)
+			];
+	};
+
+	// Helper: build one label+textbox row
+	auto MakeRow = [&](FText InLabel, TSharedPtr<SEditableTextBox>& OutBox) -> TSharedRef<SWidget>
+	{
+		return SNew(SHorizontalBox)
+			+ SHorizontalBox::Slot()
+			.AutoWidth()
+			.VAlign(VAlign_Center)
+			[
+				MakeLabel(InLabel)
+			]
+			+ SHorizontalBox::Slot()
+			.FillWidth(1.f)
+			[
+				SAssignNew(OutBox, SEditableTextBox)
+				.ClearKeyboardFocusOnCommit(false)
+			];
+	};
+
+	// Build the search form as a vertical stack of rows
+	TSharedRef<SVerticalBox> SearchForm = SNew(SVerticalBox)
+		+ SVerticalBox::Slot()
+		.AutoHeight()
+		.Padding(0.f, 0.f, 0.f, 4.f)
+		[
+			MakeRow(LOCTEXT("DescLabel", "Widget Name:"), DescSearchBox)
+		]
+
+		+ SVerticalBox::Slot()
+		.AutoHeight()
+		.Padding(0.f, 0.f, 0.f, 4.f)
+		[
+			MakeRow(LOCTEXT("TypeLabel", "Type:"), TypeSearchBox)
+		]
+
+		+ SVerticalBox::Slot()
+		.AutoHeight()
+		.Padding(0.f, 0.f, 0.f, 4.f)
+		[
+			MakeRow(LOCTEXT("TagLabel", "Tag:"), TagSearchBox)
+		]
+
+		+ SVerticalBox::Slot()
+		.AutoHeight()
+		.Padding(0.f, 0.f, 0.f, 4.f)
+		[
+			MakeRow(LOCTEXT("AddressLabel", "Address:"), AddressSearchBox)
+		]
+
+		// ── Scope selector + Search button on one row ─────────────────────
+		+ SVerticalBox::Slot()
+		.AutoHeight()
+		.Padding(0.f, 6.f, 0.f, 0.f)
+		[
+			SNew(SHorizontalBox)
+
+			// Checkbox: search all live widgets vs. current tree root
+			+ SHorizontalBox::Slot()
+			.AutoWidth()
+			.VAlign(VAlign_Center)
+			[
+				SAssignNew(AllLiveWidgetsCheckBox, SCheckBox)
+				.ToolTipText(LOCTEXT("AllLiveWidgetsTip",
+					"When checked, the search builds a fresh widget tree from every visible window "
+					"instead of searching only the nodes currently shown in the Widget Hierarchy.\n"
+					"Note: selecting a result from an all-live search will focus it in the hierarchy "
+					"only if that window has already been captured there."))
+				[
+					SNew(STextBlock)
+					.Text(LOCTEXT("AllLiveWidgetsLabel", "Search all live widgets"))
+				]
+			]
+
+			// Spacer pushes the Search button to the right
+			+ SHorizontalBox::Slot()
+			.FillWidth(1.f)
+
+			+ SHorizontalBox::Slot()
+			.AutoWidth()
+			.VAlign(VAlign_Center)
+			[
+				SNew(SButton)
+				.Text(LOCTEXT("SearchButton", "Search"))
+				.OnClicked(this, &SWidgetSearchPanel::HandleSearchClicked)
+			]
+		];
+
+	// Build the results list
+	TSharedRef<SHeaderRow> HeaderRow = SNew(SHeaderRow)
+		+ SHeaderRow::Column(WidgetSearchColumnID::DisplayName)
+		.DefaultLabel(LOCTEXT("ColWidgetName", "Widget Name"))
+		.FillWidth(2.f)
+		+ SHeaderRow::Column(WidgetSearchColumnID::Type)
+		.DefaultLabel(LOCTEXT("ColType", "Type"))
+		.FillWidth(1.5f)
+		+ SHeaderRow::Column(WidgetSearchColumnID::Tag)
+		.DefaultLabel(LOCTEXT("ColTag", "Tag"))
+		.FillWidth(1.f)
+		+ SHeaderRow::Column(WidgetSearchColumnID::Source)
+		.DefaultLabel(LOCTEXT("ColSource", "Source"))
+		.FillWidth(2.f)
+		+ SHeaderRow::Column(WidgetSearchColumnID::Address)
+		.DefaultLabel(LOCTEXT("ColAddress", "Address"))
+		.FillWidth(1.5f);
+
+	SAssignNew(ResultListView, SListView<TSharedPtr<FWidgetSearchResultEntry>>)
+		.ListItemsSource(&SearchResults)
+		.OnGenerateRow(this, &SWidgetSearchPanel::HandleGenerateRow)
+		.OnSelectionChanged(this, &SWidgetSearchPanel::HandleSelectionChanged)
+		.HeaderRow(HeaderRow)
+		.SelectionMode(ESelectionMode::Single);
+
+	ChildSlot
+	[
+		SNew(SSplitter)
+		.Orientation(Orient_Vertical)
+
+		+ SSplitter::Slot()
+		.Value(0.35f)
+		[
+			SNew(SBorder)
+			.Padding(FMargin(8.f))
+			.BorderImage(FCoreStyle::Get().GetBrush("ToolPanel.GroupBorder"))
+			[
+				SearchForm
+			]
+		]
+
+		+ SSplitter::Slot()
+		.Value(0.65f)
+		[
+			SNew(SBorder)
+			.Padding(FMargin(0.f))
+			.BorderImage(FCoreStyle::Get().GetBrush("ToolPanel.GroupBorder"))
+			[
+				ResultListView.ToSharedRef()
+			]
+		]
+	];
+}
+
+bool SWidgetSearchPanel::IsAllLiveWidgetsMode() const
+{
+	return AllLiveWidgetsCheckBox.IsValid() && AllLiveWidgetsCheckBox->IsChecked();
+}
+
+TArray<TSharedRef<FWidgetReflectorNodeBase>> SWidgetSearchPanel::BuildSearchRoots() const
+{
+	if (IsAllLiveWidgetsMode())
+	{
+		// Build a transient live tree covering every visible window.
+		TArray<TSharedRef<SWindow>> VisibleWindows;
+		FSlateApplication::Get().GetAllVisibleWindowsOrdered(VisibleWindows);
+
+		TArray<TSharedRef<FWidgetReflectorNodeBase>> Roots;
+		Roots.Reserve(VisibleWindows.Num());
+		for (const TSharedRef<SWindow>& Window : VisibleWindows)
+		{
+			Roots.Add(FWidgetReflectorNodeUtils::NewLiveNodeTreeFrom(FArrangedWidget(Window, Window->GetWindowGeometryInScreen())));
+		}
+		return Roots;
+	}
+
+	// Fall back to the current tree root provided by the reflector.
+	if (OnGetRootNodes.IsBound())
+	{
+		return OnGetRootNodes.Execute();
+	}
+	return {};
+}
+
+FReply SWidgetSearchPanel::HandleSearchClicked()
+{
+	FSearchFilters Filters;
+	Filters.Type = TypeSearchBox->GetText().ToString();
+	Filters.DisplayName = DescSearchBox->GetText().ToString();
+	Filters.Tag = TagSearchBox->GetText().ToString();
+	Filters.Address = AddressSearchBox->GetText().ToString();
+
+	// Build and cache the search roots so that result entries remain valid after the search.
+	LastSearchRoots = BuildSearchRoots();
+
+	SearchResults.Reset();
+	SearchNodes(LastSearchRoots, Filters);
+
+	if (ResultListView.IsValid())
+	{
+		ResultListView->RequestListRefresh();
+	}
+
+	return FReply::Handled();
+}
+
+void SWidgetSearchPanel::SearchNodes(const TArray<TSharedRef<FWidgetReflectorNodeBase>>& InNodes, const FSearchFilters& Filters)
+{
+	TArray<TSharedRef<FWidgetReflectorNodeBase>> Nodes = InNodes;
+	while (!Nodes.IsEmpty())
+	{
+		const TSharedRef<FWidgetReflectorNodeBase> Node = Nodes.Pop();
+		if (NodeMatchesFilters(Node, Filters))
+		{
+			SearchResults.Add(MakeShared<FWidgetSearchResultEntry>(Node));
+		}
+
+		// Always recurse into children regardless of match
+		for (const TSharedRef<FWidgetReflectorNodeBase>& Child : Node->GetChildNodes())
+		{
+			Nodes.Add(Child);
+		}
+	}
+}
+
+bool SWidgetSearchPanel::NodeMatchesFilters(const TSharedRef<FWidgetReflectorNodeBase>& Node, const FSearchFilters& Filters) const
+{
+	// If all filters are empty, nothing matches (require at least one criteria)
+	if (Filters.IsEmpty())
+	{
+		return false;
+	}
+
+	// Each non-empty filter must match (AND logic)
+	if (!Filters.Type.IsEmpty())
+	{
+		const FString NodeType = Node->GetWidgetType().ToString();
+		if (!NodeType.Contains(Filters.Type, ESearchCase::IgnoreCase))
+		{
+			return false;
+		}
+	}
+
+	if (!Filters.DisplayName.IsEmpty())
+	{
+		const FString NodeDesc = Node->GetWidgetTypeAndShortName().ToString();
+		if (!NodeDesc.Contains(Filters.DisplayName, ESearchCase::IgnoreCase))
+		{
+			return false;
+		}
+	}
+
+	if (!Filters.Tag.IsEmpty())
+	{
+		FString NodeTag;
+		if (TSharedPtr<SWidget> LiveWidget = Node->GetLiveWidget())
+		{
+			NodeTag = LiveWidget->GetTag().ToString();
+		}
+		if (!NodeTag.Contains(Filters.Tag, ESearchCase::IgnoreCase))
+		{
+			return false;
+		}
+	}
+
+	if (!Filters.Address.IsEmpty())
+	{
+		const FString NodeAddress = FWidgetReflectorNodeUtils::WidgetAddressToString(Node->GetWidgetAddress());
+		if (!NodeAddress.Contains(Filters.Address, ESearchCase::IgnoreCase))
+		{
+			return false;
+		}
+	}
+
+	return true;
+}
+
+TSharedRef<ITableRow> SWidgetSearchPanel::HandleGenerateRow(TSharedPtr<FWidgetSearchResultEntry> Item,
+	const TSharedRef<STableViewBase>& OwnerTable)
+{
+	return SNew(SWidgetSearchResultRow, OwnerTable)
+		.Entry(Item)
+		.SourceCodeAccessor(OnAccessSourceCode)
+		.AssetAccessor(OnAccessAsset);
+}
+
+void SWidgetSearchPanel::HandleSelectionChanged(TSharedPtr<FWidgetSearchResultEntry> Item, ESelectInfo::Type SelectType)
+{
+	if (!Item.IsValid())
+	{
+		return;
+	}
+
+	if (IsAllLiveWidgetsMode())
+	{
+		// In all-live-widgets mode the found nodes are transient — navigate by live widget pointer
+		// so the reflector can locate the widget in whatever tree it currently holds.
+		if (OnSelectLiveWidget.IsBound())
+		{
+			OnSelectLiveWidget.Execute(Item->Node->GetLiveWidget());
+		}
+	}
+	else
+	{
+		// In tree-root mode (searching in the current picked widgets) the found nodes are the same instance objects as in the reflector tree.
+		if (OnSelectNode.IsBound())
+		{
+			OnSelectNode.Execute(Item->Node);
+		}
+	}
+}
+
+} // namespace UE::SlateReflector
+
+#undef LOCTEXT_NAMESPACE
diff --git a/Engine/Source/Developer/SlateReflector/Private/Widgets/SWidgetSearchPanel.h b/Engine/Source/Developer/SlateReflector/Private/Widgets/SWidgetSearchPanel.h
new file mode 100644
--- /dev/null
+++ b/Engine/Source/Developer/SlateReflector/Private/Widgets/SWidgetSearchPanel.h
@@ -0,0 +1,98 @@
+// Copyright Epic Games, Inc. All Rights Reserved.
+
+#pragma once
+
+#include "CoreMinimal.h"
+#include "Framework/Application/IWidgetReflector.h"
+#include "Widgets/DeclarativeSyntaxSupport.h"
+#include "Widgets/Input/SCheckBox.h"
+#include "Widgets/SCompoundWidget.h"
+#include "Widgets/Views/SListView.h"
+
+class FWidgetReflectorNodeBase;
+class SEditableTextBox;
+
+namespace UE::SlateReflector
+{
+struct FWidgetSearchResultEntry;
+
+/**
+ * A panel that allows searching the widget hierarchy by type, description, tag, and address.
+ * Displays matching widgets in a list view; selecting an entry focuses it in the Widget Hierarchy.
+ *
+ * Two search scopes are available via a checkbox:
+ *  - Current tree root  : searches the nodes currently shown in the Widget Hierarchy.
+ *  - All live widgets   : builds a transient FWidgetReflectorNodeBase tree from every visible
+ *                         SWindow so that widgets outside the current pick scope can be found.
+ */
+class SWidgetSearchPanel : public SCompoundWidget
+{
+public:
+	/** Called when the user selects a result that belongs to the current tree root. */
+	DECLARE_DELEGATE_OneParam(FOnSelectNode, TSharedRef<FWidgetReflectorNodeBase>);
+
+	/** Called when the user selects a result in "all live widgets" mode; passes the live SWidget
+	 *  so the hierarchy can locate it even if the tree root hasn't been set for that window yet. */
+	DECLARE_DELEGATE_OneParam(FOnSelectLiveWidget, TSharedPtr<const SWidget>);
+
+	/** Called to obtain the current root nodes (live or snapshot) at the time of search. */
+	DECLARE_DELEGATE_RetVal(TArray<TSharedRef<FWidgetReflectorNodeBase>>, FGetRootNodes);
+
+	SLATE_BEGIN_ARGS(SWidgetSearchPanel) {}
+		SLATE_EVENT(FOnSelectNode,       OnSelectNode)
+		SLATE_EVENT(FOnSelectLiveWidget, OnSelectLiveWidget)
+		SLATE_EVENT(FGetRootNodes,       OnGetRootNodes)
+		SLATE_EVENT(FAccessSourceCode,   OnAccessSourceCode)
+		SLATE_EVENT(FAccessAsset,        OnAccessAsset)
+	SLATE_END_ARGS()
+
+	void Construct(const FArguments& InArgs);
+
+private:
+	struct FSearchFilters
+	{
+		FString DisplayName;
+		FString Type;
+		FString Tag;
+		FString Address;
+
+		bool IsEmpty() const
+		{
+			return Type.IsEmpty() && DisplayName.IsEmpty() && Tag.IsEmpty() && Address.IsEmpty();
+		}
+	};
+
+	FReply HandleSearchClicked();
+
+	/** Build search roots: either the current tree root or a fresh live tree from all windows. */
+	TArray<TSharedRef<FWidgetReflectorNodeBase>> BuildSearchRoots() const;
+
+	void SearchNodes(const TArray<TSharedRef<FWidgetReflectorNodeBase>>& InNodes, const FSearchFilters& Filters);
+	bool NodeMatchesFilters(const TSharedRef<FWidgetReflectorNodeBase>& Node, const FSearchFilters& Filters) const;
+
+	TSharedRef<ITableRow> HandleGenerateRow(TSharedPtr<FWidgetSearchResultEntry> Item, const TSharedRef<STableViewBase>& OwnerTable);
+	void HandleSelectionChanged(TSharedPtr<FWidgetSearchResultEntry> Item, ESelectInfo::Type SelectType);
+
+	/** Returns true if the "Search all live widgets" checkbox is currently checked. */
+	bool IsAllLiveWidgetsMode() const;
+
+	TSharedPtr<SEditableTextBox> TypeSearchBox;
+	TSharedPtr<SEditableTextBox> DescSearchBox;
+	TSharedPtr<SEditableTextBox> TagSearchBox;
+	TSharedPtr<SEditableTextBox> AddressSearchBox;
+
+	TSharedPtr<SCheckBox> AllLiveWidgetsCheckBox;
+
+	TArray<TSharedRef<FWidgetReflectorNodeBase>> LastSearchRoots;
+
+	TArray<TSharedPtr<FWidgetSearchResultEntry>> SearchResults;
+	TSharedPtr<SListView<TSharedPtr<FWidgetSearchResultEntry>>> ResultListView;
+
+	FOnSelectNode OnSelectNode;
+	FOnSelectLiveWidget OnSelectLiveWidget;
+	FGetRootNodes     OnGetRootNodes;
+	FAccessSourceCode OnAccessSourceCode;
+	FAccessAsset      OnAccessAsset;
+};
+
+} // namespace UE::SlateReflector
